Fixing function parameter export code in Script Generator Plugin

diff --git a/Engine/Plugins/Script/ScriptGeneratorPlugin/Source/ScriptGeneratorPlugin/Private/GenericScriptCodeGenerator.cpp b/Engine/Plugins/Script/ScriptGeneratorPlugin/Source/ScriptGeneratorPlugin/Private/GenericScriptCodeGenerator.cpp
--- a/Engine/Plugins/Script/ScriptGeneratorPlugin/Source/ScriptGeneratorPlugin/Private/GenericScriptCodeGenerator.cpp
+++ b/Engine/Plugins/Script/ScriptGeneratorPlugin/Source/ScriptGeneratorPlugin/Private/GenericScriptCodeGenerator.cpp
@@ -27,9 +27,9 @@ FString FGenericScriptCodeGenerator::GenerateFunctionParamDeclaration(const FStr
 	}
 	else
 	{
-		Initializer = FString::Printf(TEXT("%s()"), *Param->GetCPPType(NULL, CPPF_ArgumentOrReturnValue));
+		Initializer = FString::Printf(TEXT("%s()"), *GetPropertyTypeCPP(Param, CPPF_ArgumentOrReturnValue));
 	}
-	return FString::Printf(TEXT("%s %s = %s;"), *Param->GetCPPType(NULL, CPPF_ArgumentOrReturnValue), *Param->GetName(), *Initializer);
+	return FString::Printf(TEXT("%s %s = %s;"), *GetPropertyTypeCPP(Param, CPPF_ArgumentOrReturnValue), *Param->GetName(), *Initializer);
 }
 
 FString FGenericScriptCodeGenerator::GenerateObjectDeclarationFromContext(const FString& ClassNameCPP, UClass* Class)
@@ -76,7 +76,7 @@ FString FGenericScriptCodeGenerator::ExportFunction(const FString& ClassNameCPP,
 			else
 			{
 				ReturnValue = Param;
-				ReturnValueDeclaration += FString::Printf(TEXT("%s %s = "), *Param->GetCPPType(NULL, CPPF_ArgumentOrReturnValue), *Param->GetName());
+				ReturnValueDeclaration += FString::Printf(TEXT("%s %s = "), *GetPropertyTypeCPP(Param, CPPF_ArgumentOrReturnValue), *Param->GetName());
 			}
 		}
 		FunctionBody += FString::Printf(TEXT("\t%s\r\n"), *GenerateObjectDeclarationFromContext(ClassNameCPP, Class));
diff --git a/Engine/Plugins/Script/ScriptGeneratorPlugin/Source/ScriptGeneratorPlugin/Private/LuaScriptCodeGenerator.cpp b/Engine/Plugins/Script/ScriptGeneratorPlugin/Source/ScriptGeneratorPlugin/Private/LuaScriptCodeGenerator.cpp
--- a/Engine/Plugins/Script/ScriptGeneratorPlugin/Source/ScriptGeneratorPlugin/Private/LuaScriptCodeGenerator.cpp
+++ b/Engine/Plugins/Script/ScriptGeneratorPlugin/Source/ScriptGeneratorPlugin/Private/LuaScriptCodeGenerator.cpp
@@ -79,14 +79,14 @@ FString FLuaScriptCodeGenerator::GenerateFunctionParamDeclaration(const FString&
 	}
 	else if (Param->IsA(UObjectPropertyBase::StaticClass()))
 	{
-		Initializer = FString::Printf(TEXT("(%s)(lua_touserdata"), *Param->GetCPPType(NULL, CPPF_ArgumentOrReturnValue), ParamIndex);
+		Initializer = FString::Printf(TEXT("(%s)(lua_touserdata"), *GetPropertyTypeCPP(Param, CPPF_ArgumentOrReturnValue), ParamIndex);
 	}
 	else
 	{
 		FError::Throwf(TEXT("Unsupported function param type: %s"), *Param->GetClass()->GetName());
 	}
 
-	return FString::Printf(TEXT("%s %s = %s(InScriptContext, %d));"), *Param->GetCPPType(NULL, CPPF_ArgumentOrReturnValue), *Param->GetName(), *Initializer, ParamIndex);
+	return FString::Printf(TEXT("%s %s = %s(InScriptContext, %d));"), *GetPropertyTypeCPP(Param, CPPF_ArgumentOrReturnValue), *Param->GetName(), *Initializer, ParamIndex);
 }
 
 FString FLuaScriptCodeGenerator::GenerateObjectDeclarationFromContext(const FString& ClassNameCPP, UClass* Class)
@@ -255,7 +255,7 @@ FString FLuaScriptCodeGenerator::ExportFunction(const FString& ClassNameCPP, UCl
 			else
 			{
 				ReturnValue = Param;
-				ReturnValueDeclaration += FString::Printf(TEXT("%s %s = "), *Param->GetCPPType(NULL, CPPF_ArgumentOrReturnValue), *Param->GetName());
+				ReturnValueDeclaration += FString::Printf(TEXT("%s %s = "), *GetPropertyTypeCPP(Param, CPPF_ArgumentOrReturnValue), *Param->GetName());
 			}
 		}		
 		FunctionBody += FString::Printf(TEXT("\t%sObj->%s(%s);\r\n"), *ReturnValueDeclaration, *Function->GetName(), *FunctionCallArguments);
diff --git a/Engine/Plugins/Script/ScriptGeneratorPlugin/Source/ScriptGeneratorPlugin/Private/ScriptCodeGeneratorBase.cpp b/Engine/Plugins/Script/ScriptGeneratorPlugin/Source/ScriptGeneratorPlugin/Private/ScriptCodeGeneratorBase.cpp
--- a/Engine/Plugins/Script/ScriptGeneratorPlugin/Source/ScriptGeneratorPlugin/Private/ScriptCodeGeneratorBase.cpp
+++ b/Engine/Plugins/Script/ScriptGeneratorPlugin/Source/ScriptGeneratorPlugin/Private/ScriptCodeGeneratorBase.cpp
@@ -70,6 +70,34 @@ FString FScriptCodeGeneratorBase::GetClassNameCPP(UClass* Class)
 	return FString::Printf(TEXT("%s%s"), Class->GetPrefixCPP(), *Class->GetName());
 }
 
+FString FScriptCodeGeneratorBase::GetPropertyTypeCPP(UProperty* Property, uint32 PortFlags /*= 0*/)
+{
+	static FString EnumDecl(TEXT("enum "));
+	static FString StructDecl(TEXT("struct "));
+	static FString ClassDecl(TEXT("class "));
+	static FString TEnumAsByteDecl(TEXT("TEnumAsByte<enum "));
+	static FString TSubclassOfDecl(TEXT("TSubclassOf<class "));
+
+	FString PropertyType = Property->GetCPPType(NULL, PortFlags);
+	// Strip any forward declaration keywords
+	if (PropertyType.StartsWith(EnumDecl) || PropertyType.StartsWith(StructDecl) || PropertyType.StartsWith(ClassDecl))
+	{
+		int FirstSpaceIndex = PropertyType.Find(TEXT(" "));
+		PropertyType = PropertyType.Mid(FirstSpaceIndex + 1);
+	}
+	else if (PropertyType.StartsWith(TEnumAsByteDecl))
+	{
+		int FirstSpaceIndex = PropertyType.Find(TEXT(" "));
+		PropertyType = TEXT("TEnumAsByte<") + PropertyType.Mid(FirstSpaceIndex + 1);
+	}
+	else if (PropertyType.StartsWith(TSubclassOfDecl))
+	{
+		int FirstSpaceIndex = PropertyType.Find(TEXT(" "));
+		PropertyType = TEXT("TSubclassOf<") + PropertyType.Mid(FirstSpaceIndex + 1);
+	}
+	return PropertyType;
+}
+
 FString FScriptCodeGeneratorBase::GetScriptHeaderForClass(UClass* Class)
 {
 	return GeneratedCodePath / (Class->GetName() + TEXT(".script.h"));
diff --git a/Engine/Plugins/Script/ScriptGeneratorPlugin/Source/ScriptGeneratorPlugin/Private/ScriptCodeGeneratorBase.h b/Engine/Plugins/Script/ScriptGeneratorPlugin/Source/ScriptGeneratorPlugin/Private/ScriptCodeGeneratorBase.h
--- a/Engine/Plugins/Script/ScriptGeneratorPlugin/Source/ScriptGeneratorPlugin/Private/ScriptCodeGeneratorBase.h
+++ b/Engine/Plugins/Script/ScriptGeneratorPlugin/Source/ScriptGeneratorPlugin/Private/ScriptCodeGeneratorBase.h
@@ -32,6 +32,8 @@ protected:
 	FString RebaseToBuildPath(const FString& FileName) const;
 	/** Converts a UClass name to C++ class name (with U/A prefix) */
 	FString GetClassNameCPP(UClass* Class);
+	/** Gets C++ friendly property type name. */
+	FString GetPropertyTypeCPP(UProperty* Property, uint32 PortFlags = 0);
 	/** Generates a script glue header filename for the specified class */
 	virtual FString GetScriptHeaderForClass(UClass* Class);
 	/** Returns true if the specified class can be exported */
diff --git a/Engine/Programs/UnrealHeaderTool/Config/DefaultEngine.ini b/Engine/Programs/UnrealHeaderTool/Config/DefaultEngine.ini
--- a/Engine/Programs/UnrealHeaderTool/Config/DefaultEngine.ini
+++ b/Engine/Programs/UnrealHeaderTool/Config/DefaultEngine.ini
@@ -3,4 +3,4 @@
 +ScriptPlugins="ScriptGeneratorPlugin"
 
 +ScriptSupportedModules=/Script/CoreUObject
-;+ScriptSupportedModules=/Script/Engine
\ No newline at end of file
++ScriptSupportedModules=/Script/Engine
\ No newline at end of file
