Script plugin improvements: functions are now called with ProcessEvent, properties are being exported too, new test actor class, BeginPlay event.

diff --git a/Engine/Plugins/Script/ScriptGeneratorPlugin/Source/ScriptGeneratorPlugin/Private/GenericScriptCodeGenerator.cpp b/Engine/Plugins/Script/ScriptGeneratorPlugin/Source/ScriptGeneratorPlugin/Private/GenericScriptCodeGenerator.cpp
--- a/Engine/Plugins/Script/ScriptGeneratorPlugin/Source/ScriptGeneratorPlugin/Private/GenericScriptCodeGenerator.cpp
+++ b/Engine/Plugins/Script/ScriptGeneratorPlugin/Source/ScriptGeneratorPlugin/Private/GenericScriptCodeGenerator.cpp
@@ -34,7 +34,7 @@ FString FGenericScriptCodeGenerator::GenerateFunctionParamDeclaration(const FStr
 
 FString FGenericScriptCodeGenerator::GenerateObjectDeclarationFromContext(const FString& ClassNameCPP, UClass* Class)
 {
-	return FString::Printf(TEXT("%s* Obj = (%s*)InScriptContext;"), *ClassNameCPP, *ClassNameCPP);
+	return FString::Printf(TEXT("UObject* Obj = (%s*)InScriptContext;"), *ClassNameCPP);
 }
 
 FString FGenericScriptCodeGenerator::GenerateReturnValueHandler(const FString& ClassNameCPP, UClass* Class, UFunction* Function, UProperty* ReturnValue)
@@ -62,25 +62,19 @@ FString FGenericScriptCodeGenerator::ExportFunction(const FString& ClassNameCPP,
 	FString FunctionBody;
 	if (FuncSuper == NULL)
 	{
+		FunctionBody += FString::Printf(TEXT("\t%s\r\n"), *GenerateObjectDeclarationFromContext(ClassNameCPP, Class));
+		FunctionBody += GenerateFunctionDispatch(Function);
+
 		FString FunctionCallArguments;
 		FString ReturnValueDeclaration;
-
-		for (TFieldIterator<UProperty> ParamIt(Function); ParamIt; ++ParamIt)
+		for (TFieldIterator<UProperty> ParamIt(Function); !ReturnValue && ParamIt; ++ParamIt)
 		{
 			UProperty* Param = *ParamIt;
-			if (!(Param->GetPropertyFlags() & CPF_ReturnParm))
-			{
-				FunctionBody += FString::Printf(TEXT("\t%s\r\n"), *GenerateFunctionParamDeclaration(ClassNameCPP, Class, Function, Param));
-				FunctionCallArguments += (FunctionCallArguments.Len() > 0) ? FString::Printf(TEXT(", %s"), *Param->GetName()) : Param->GetName();
-			}
-			else
+			if (Param->GetPropertyFlags() & CPF_ReturnParm)
 			{
 				ReturnValue = Param;
-				ReturnValueDeclaration += FString::Printf(TEXT("%s %s = "), *GetPropertyTypeCPP(Param, CPPF_ArgumentOrReturnValue), *Param->GetName());
 			}
 		}
-		FunctionBody += FString::Printf(TEXT("\t%s\r\n"), *GenerateObjectDeclarationFromContext(ClassNameCPP, Class));
-		FunctionBody += FString::Printf(TEXT("\t%sObj->%s(%s);\r\n"), *ReturnValueDeclaration, *Function->GetName(), *FunctionCallArguments);
 		FunctionBody += FString::Printf(TEXT("\t%s\r\n"), *GenerateReturnValueHandler(ClassNameCPP, Class, Function, ReturnValue));
 	}
 	else
@@ -117,17 +111,17 @@ FString FGenericScriptCodeGenerator::ExportProperty(const FString& ClassNameCPP,
 	if (PropertySuper == NULL)
 	{
 		FunctionBody += FString::Printf(TEXT("\t%s\r\n"), *GenerateObjectDeclarationFromContext(ClassNameCPP, Class));
-		FunctionBody += FString::Printf(TEXT("\tUProperty* Property = %s_GetExportedProperty(%d);\r\n"), *Class->GetName(), PropertyIndex);
+		FunctionBody += FString::Printf(TEXT("\tstatic UProperty* Property = FindScriptPropertyHelper(%s::StaticClass(), TEXT(\"%s\"));\r\n"), *ClassNameCPP, *Property->GetName());
 		FunctionBody += FString::Printf(TEXT("\t%s\r\n"), *GenerateFunctionParamDeclaration(ClassNameCPP, Class, NULL, Property));
 		FunctionBody += FString::Printf(TEXT("\tProperty->CopyCompleteValue(&%s, Property->ContainerPtrToValuePtr<void>(Obj));\r\n"), *Property->GetName());
 		FunctionBody += FString::Printf(TEXT("\t// @todo: handle property value here\r\n"));
+		FunctionBody += TEXT("\treturn 0;\r\n");
 	}
 	else
 	{
 		FunctionBody = FString::Printf(TEXT("\treturn %s_%s(InScriptContext);\r\n"), *PropertySuper->GetName(), *GetterName);
 	}
-	GeneratedGlue += FunctionBody;
-	GeneratedGlue += TEXT("\treturn 0;\r\n");
+	GeneratedGlue += FunctionBody;	
 	GeneratedGlue += TEXT("}\r\n\r\n");
 	FunctionBody.Empty(FunctionBody.Len());
 
@@ -138,16 +132,16 @@ FString FGenericScriptCodeGenerator::ExportProperty(const FString& ClassNameCPP,
 	if (PropertySuper == NULL)
 	{
 		FunctionBody += FString::Printf(TEXT("\t%s\r\n"), *GenerateObjectDeclarationFromContext(ClassNameCPP, Class));
-		FunctionBody += FString::Printf(TEXT("\tUProperty* Property = %s_GetExportedProperty(%d);\r\n"), *Class->GetName(), PropertyIndex);
+		FunctionBody += FString::Printf(TEXT("\tstatic UProperty* Property = FindScriptPropertyHelper(%s::StaticClass(), TEXT(\"%s\"));\r\n"), *ClassNameCPP, *Property->GetName());
 		FunctionBody += FString::Printf(TEXT("\t%s\r\n"), *GenerateFunctionParamDeclaration(ClassNameCPP, Class, NULL, Property));
 		FunctionBody += FString::Printf(TEXT("\tProperty->CopyCompleteValue(Property->ContainerPtrToValuePtr<void>(Obj), &%s);\r\n"), *Property->GetName());
+		FunctionBody += TEXT("\treturn 0;\r\n");
 	}
 	else
 	{
 		FunctionBody = FString::Printf(TEXT("\treturn %s_%s(InScriptContext);\r\n"), *PropertySuper->GetName(), *SetterName);
 	}
-	GeneratedGlue += FunctionBody;
-	GeneratedGlue += TEXT("\treturn 0;\r\n");
+	GeneratedGlue += FunctionBody;	
 	GeneratedGlue += TEXT("}\r\n\r\n");
 
 	return GeneratedGlue;
@@ -161,7 +155,9 @@ FString FGenericScriptCodeGenerator::ExportAdditionalClassGlue(const FString& Cl
 	{
 		GeneratedGlue += GenerateWrapperFunctionDeclaration(ClassNameCPP, Class, TEXT("New"));
 		GeneratedGlue += TEXT("\r\n{\r\n");
-		GeneratedGlue += FString::Printf(TEXT("\tUObject* Obj = NewObject<%s>();\r\n"), *ClassNameCPP);
+		GeneratedGlue += TEXT("\tUObject* Outer = NULL;\r\n");
+		GeneratedGlue += TEXT("\tFName Name = FName(\"ScriptObject\");\r\n");
+		GeneratedGlue += FString::Printf(TEXT("\tUObject* Obj = NewNamedObject<%s>(Outer, Name);\r\n"), *ClassNameCPP);
 		GeneratedGlue += TEXT("\tif (Obj)\r\n\t{\r\n");
 		GeneratedGlue += TEXT("\t\tFScriptObjectReferencer::Get().AddObjectReference(Obj);\r\n");
 		GeneratedGlue += TEXT("\t\t// @todo: Register the object with the script context here\r\n");
@@ -257,14 +253,12 @@ void FGenericScriptCodeGenerator::ExportClass(UClass* Class, const FString& Sour
 		}
 	}
 
-	FString ExportedPropertiesLookupFunction;
-	TArray<UProperty*> ExportedProperties;
-	if (GenerateExportedPropertyTable(ClassNameCPP, Class, ExportedPropertiesLookupFunction, ExportedProperties) > 0)
+	int32 PropertyIndex = 0;
+	for (TFieldIterator<UProperty> PropertyIt(Class, EFieldIteratorFlags::ExcludeSuper); PropertyIt; ++PropertyIt, ++PropertyIndex)
 	{
-		GeneratedGlue += ExportedPropertiesLookupFunction;
-		for (int32 PropertyIndex = 0; PropertyIndex < ExportedProperties.Num(); ++PropertyIndex)
+		UProperty* Property = *PropertyIt;
+		if (CanExportProperty(ClassNameCPP, Class, Property))
 		{
-			UProperty* Property = ExportedProperties[PropertyIndex];
 			UE_LOG(LogScriptGenerator, Log, TEXT("  %s %s"), *Property->GetClass()->GetName(), *Property->GetName());
 			GeneratedGlue += ExportProperty(ClassNameCPP, Class, Property, PropertyIndex);
 		}
diff --git a/Engine/Plugins/Script/ScriptGeneratorPlugin/Source/ScriptGeneratorPlugin/Private/LuaScriptCodeGenerator.cpp b/Engine/Plugins/Script/ScriptGeneratorPlugin/Source/ScriptGeneratorPlugin/Private/LuaScriptCodeGenerator.cpp
--- a/Engine/Plugins/Script/ScriptGeneratorPlugin/Source/ScriptGeneratorPlugin/Private/LuaScriptCodeGenerator.cpp
+++ b/Engine/Plugins/Script/ScriptGeneratorPlugin/Source/ScriptGeneratorPlugin/Private/LuaScriptCodeGenerator.cpp
@@ -24,81 +24,88 @@ FString FLuaScriptCodeGenerator::GenerateWrapperFunctionDeclaration(const FStrin
 	return FString::Printf(TEXT("int32 %s_%s(lua_State* InScriptContext)"), *Class->GetName(), *FunctionName);
 }
 
-FString FLuaScriptCodeGenerator::GenerateFunctionParamDeclaration(const FString& ClassNameCPP, UClass* Class, UFunction* Function, UProperty* Param, int32 ParamIndex)
-{
-	// In Lua, the first param index on the stack is 1 and it's the object we're invoking the function on
-	ParamIndex += 2;
-	FString Initializer;
-	if (Param->IsA(UIntProperty::StaticClass()))
-	{
-		Initializer = TEXT("(luaL_checkint");
-	}
-	else if (Param->IsA(UFloatProperty::StaticClass()))
-	{
-		Initializer = TEXT("(float)(luaL_checknumber");
-	}
-	else if (Param->IsA(UStrProperty::StaticClass()))
-	{
-		Initializer = TEXT("ANSI_TO_TCHAR(luaL_checkstring");
-	}
-	else if (Param->IsA(UNameProperty::StaticClass()))
+FString FLuaScriptCodeGenerator::InitializeFunctionDispatchParam(UFunction* Function, UProperty* Param, int32 ParamIndex)
+{	
+	if (!(Param->GetPropertyFlags() & CPF_ReturnParm))
 	{
-		Initializer = TEXT("FName(luaL_checkstring");
-	}
-	else if (Param->IsA(UBoolProperty::StaticClass()))
-	{
-		Initializer = TEXT("!!(lua_toboolean");
-	}
-	else if (Param->IsA(UStructProperty::StaticClass()))
-	{
-		UStructProperty* StructProp = CastChecked<UStructProperty>(Param);
-		if (StructProp->Struct->GetFName() == Name_Vector)
+		FString Initializer;
+		// In Lua, the first param index on the stack is 1 and it's the object we're invoking the function on
+		ParamIndex += 2;
+
+		if (Param->IsA(UIntProperty::StaticClass()))
+		{
+			Initializer = TEXT("(luaL_checkint");
+		}
+		else if (Param->IsA(UFloatProperty::StaticClass()))
 		{
-			Initializer = TEXT("(FLuaVector::Get");
+			Initializer = TEXT("(float)(luaL_checknumber");
 		}
-		else if (StructProp->Struct->GetFName() == Name_Vector4)
+		else if (Param->IsA(UStrProperty::StaticClass()))
 		{
-			Initializer = TEXT("(FLuaVector4::Get");
+			Initializer = TEXT("ANSI_TO_TCHAR(luaL_checkstring");
+		}
+		else if (Param->IsA(UNameProperty::StaticClass()))
+		{
+			Initializer = TEXT("FName(luaL_checkstring");
+		}
+		else if (Param->IsA(UBoolProperty::StaticClass()))
+		{
+			Initializer = TEXT("!!(lua_toboolean");
+		}
+		else if (Param->IsA(UStructProperty::StaticClass()))
+		{
+			UStructProperty* StructProp = CastChecked<UStructProperty>(Param);
+			if (StructProp->Struct->GetFName() == Name_Vector)
+			{
+				Initializer = TEXT("(FLuaVector::Get");
+			}
+			else if (StructProp->Struct->GetFName() == Name_Vector4)
+			{
+				Initializer = TEXT("(FLuaVector4::Get");
+			}
+			else if (StructProp->Struct->GetFName() == Name_Quat)
+			{
+				Initializer = TEXT("(FLuaQuat::Get");
+			}
+			else if (StructProp->Struct->GetFName() == Name_Transform)
+			{
+				Initializer = TEXT("(FLuaTransform::Get");
+			}
+			else
+			{
+				FError::Throwf(TEXT("Unsupported function param struct type: %s"), *StructProp->Struct->GetName());
+			}
 		}
-		else if (StructProp->Struct->GetFName() == Name_Quat)
+		else if (Param->IsA(UClassProperty::StaticClass()))
 		{
-			Initializer = TEXT("(FLuaQuat::Get");
+			Initializer = TEXT("(UClass*)(lua_touserdata");
 		}
-		else if (StructProp->Struct->GetFName() == Name_Transform)
+		else if (Param->IsA(UObjectPropertyBase::StaticClass()))
 		{
-			Initializer = TEXT("(FLuaTransform::Get");
+			Initializer = FString::Printf(TEXT("(%s)(lua_touserdata"), *GetPropertyTypeCPP(Param, CPPF_ArgumentOrReturnValue), ParamIndex);
 		}
 		else
 		{
-			FError::Throwf(TEXT("Unsupported function param struct type: %s"), *StructProp->Struct->GetName());
+			FError::Throwf(TEXT("Unsupported function param type: %s"), *Param->GetClass()->GetName());
 		}
-	}
-	else if (Param->IsA(UClassProperty::StaticClass()))
-	{
-		Initializer = TEXT("(UClass*)(lua_touserdata");
-	}
-	else if (Param->IsA(UObjectPropertyBase::StaticClass()))
-	{
-		Initializer = FString::Printf(TEXT("(%s)(lua_touserdata"), *GetPropertyTypeCPP(Param, CPPF_ArgumentOrReturnValue), ParamIndex);
+
+		return FString::Printf(TEXT("%s(InScriptContext, %d))"), *Initializer, ParamIndex);
 	}
 	else
 	{
-		FError::Throwf(TEXT("Unsupported function param type: %s"), *Param->GetClass()->GetName());
-	}
-
-	return FString::Printf(TEXT("%s %s = %s(InScriptContext, %d));"), *GetPropertyTypeCPP(Param, CPPF_ArgumentOrReturnValue), *Param->GetName(), *Initializer, ParamIndex);
+		return FScriptCodeGeneratorBase::InitializeFunctionDispatchParam(Function, Param, ParamIndex);
+	}	
 }
 
 FString FLuaScriptCodeGenerator::GenerateObjectDeclarationFromContext(const FString& ClassNameCPP, UClass* Class)
 {
-	return FString::Printf(TEXT("%s* Obj = (%s*)lua_touserdata(InScriptContext, 1);"), *ClassNameCPP, *ClassNameCPP);
+	return FString::Printf(TEXT("UObject* Obj = (%s*)lua_touserdata(InScriptContext, 1);"), *ClassNameCPP);
 }
 
-FString FLuaScriptCodeGenerator::GenerateReturnValueHandler(const FString& ClassNameCPP, UClass* Class, UFunction* Function, UProperty* ReturnValue)
+FString FLuaScriptCodeGenerator::GenerateReturnValueHandler(const FString& ClassNameCPP, UClass* Class, UFunction* Function, UProperty* ReturnValue, const FString& ReturnValueName)
 {
 	if (ReturnValue)
 	{
-		const FString ReturnValueName = ReturnValue->GetName();
 		FString Initializer;		
 		if (ReturnValue->IsA(UIntProperty::StaticClass()))
 		{
@@ -178,6 +185,18 @@ bool FLuaScriptCodeGenerator::CanExportClass(UClass* Class)
 				{
 					bHasMembersToExport = true;
 				}
+			}			
+			// Check properties too
+			if (!bHasMembersToExport)
+			{
+				for (TFieldIterator<UProperty> PropertyIt(Class, EFieldIteratorFlags::ExcludeSuper); !bHasMembersToExport && PropertyIt; ++PropertyIt)
+				{
+					UProperty* Property = *PropertyIt;
+					if (CanExportProperty(ClassNameCPP, Class, Property))
+					{
+						bHasMembersToExport = true;
+					}
+				}
 			}
 			bCanExport = bHasMembersToExport;
 		}
@@ -193,27 +212,7 @@ bool FLuaScriptCodeGenerator::CanExportFunction(const FString& ClassNameCPP, UCl
 		for (TFieldIterator<UProperty> ParamIt(Function); bExport && ParamIt; ++ParamIt)
 		{
 			UProperty* Param = *ParamIt;
-			if (Param->IsA(UStructProperty::StaticClass()))
-			{
-				UStructProperty* StructProp = CastChecked<UStructProperty>(Param);
-				if (StructProp->Struct->GetFName() != Name_Vector &&
-					StructProp->Struct->GetFName() != Name_Vector4 &&
-					StructProp->Struct->GetFName() != Name_Quat &&
-					StructProp->Struct->GetFName() != Name_Transform)
-				{
-					bExport = false;
-				}
-			}
-			else if (!Param->IsA(UIntProperty::StaticClass()) &&
-				  !Param->IsA(UFloatProperty::StaticClass()) &&
-					!Param->IsA(UStrProperty::StaticClass()) &&
-					!Param->IsA(UNameProperty::StaticClass()) &&
-					!Param->IsA(UBoolProperty::StaticClass()) &&
-					!Param->IsA(UObjectPropertyBase::StaticClass()) &&
-					!Param->IsA(UClassProperty::StaticClass()))
-			{
-				bExport = false;
-			}
+			bExport = IsPropertyTypeSupported(Param);
 		}
 	}
 	return bExport;
@@ -240,26 +239,24 @@ FString FLuaScriptCodeGenerator::ExportFunction(const FString& ClassNameCPP, UCl
 	if (FuncSuper == NULL)
 	{
 		FunctionBody += FString::Printf(TEXT("\t%s\r\n"), *GenerateObjectDeclarationFromContext(ClassNameCPP, Class));
+		FunctionBody += GenerateFunctionDispatch(Function);
 
 		FString FunctionCallArguments;
 		FString ReturnValueDeclaration;
-		int32 ParamIndex = 0;
-		for (TFieldIterator<UProperty> ParamIt(Function); ParamIt; ++ParamIt, ++ParamIndex)
+		for (TFieldIterator<UProperty> ParamIt(Function); !ReturnValue && ParamIt; ++ParamIt)
 		{
 			UProperty* Param = *ParamIt;
-			if (!(Param->GetPropertyFlags() & CPF_ReturnParm))
-			{
-				FunctionBody += FString::Printf(TEXT("\t%s\r\n"), *GenerateFunctionParamDeclaration(ClassNameCPP, Class, Function, Param, ParamIndex));
-				FunctionCallArguments += (FunctionCallArguments.Len() > 0) ? FString::Printf(TEXT(", %s"), *Param->GetName()) : Param->GetName();
-			}
-			else
+			if (Param->GetPropertyFlags() & CPF_ReturnParm)
 			{
 				ReturnValue = Param;
-				ReturnValueDeclaration += FString::Printf(TEXT("%s %s = "), *GetPropertyTypeCPP(Param, CPPF_ArgumentOrReturnValue), *Param->GetName());
 			}
-		}		
-		FunctionBody += FString::Printf(TEXT("\t%sObj->%s(%s);\r\n"), *ReturnValueDeclaration, *Function->GetName(), *FunctionCallArguments);
-		FunctionBody += FString::Printf(TEXT("\t%s\r\n"), *GenerateReturnValueHandler(ClassNameCPP, Class, Function, ReturnValue));
+		}
+		FString ReturnValueName;
+		if (ReturnValue)
+		{
+			ReturnValueName = FString::Printf(TEXT("Params.%s"), *ReturnValue->GetName());
+		}
+		FunctionBody += FString::Printf(TEXT("\t%s\r\n"), *GenerateReturnValueHandler(ClassNameCPP, Class, Function, ReturnValue, *ReturnValueName));
 	}
 	else
 	{
@@ -275,9 +272,116 @@ FString FLuaScriptCodeGenerator::ExportFunction(const FString& ClassNameCPP, UCl
 	return GeneratedGlue;
 }
 
-FString FLuaScriptCodeGenerator::ExportProperty(const FString& ClassNameCPP, UClass* Class, UProperty* Property)
+bool FLuaScriptCodeGenerator::IsPropertyTypeSupported(UProperty* Property) const
 {
-	FString GeneratedGlue;
+	bool bSupported = true;
+	if (Property->IsA(UStructProperty::StaticClass()))
+	{
+		UStructProperty* StructProp = CastChecked<UStructProperty>(Property);
+		if (StructProp->Struct->GetFName() != Name_Vector &&
+			StructProp->Struct->GetFName() != Name_Vector4 &&
+			StructProp->Struct->GetFName() != Name_Quat &&
+			StructProp->Struct->GetFName() != Name_Transform)
+		{
+			bSupported = false;
+		}
+	}
+	else if (!Property->IsA(UIntProperty::StaticClass()) &&
+		!Property->IsA(UFloatProperty::StaticClass()) &&
+		!Property->IsA(UStrProperty::StaticClass()) &&
+		!Property->IsA(UNameProperty::StaticClass()) &&
+		!Property->IsA(UBoolProperty::StaticClass()) &&
+		!Property->IsA(UObjectPropertyBase::StaticClass()) &&
+		!Property->IsA(UClassProperty::StaticClass()))
+	{
+		bSupported = false;
+	}
+	else if (Property->IsA(ULazyObjectProperty::StaticClass()))
+	{
+		bSupported = false;
+	}
+	return bSupported;
+}
+
+bool FLuaScriptCodeGenerator::CanExportProperty(const FString& ClassNameCPP, UClass* Class, UProperty* Property)
+{
+	// Only editable properties can be exported
+	if (!(Property->PropertyFlags & CPF_Edit))
+	{
+		return false;
+	}
+	// Check if property type is supported
+	return IsPropertyTypeSupported(Property);
+}
+
+FString FLuaScriptCodeGenerator::ExportProperty(const FString& ClassNameCPP, UClass* Class, UProperty* Property, int32 PropertyIndex)
+{
+	FString PropertyName = Property->GetName();
+	UProperty* ReturnValue = NULL;
+	UClass* PropertySuper = NULL;
+
+	if (Property->GetOwnerClass() != Class)
+	{
+		// Find the base class where this property was defined
+		if (ExportedClasses.Contains(Property->GetOwnerClass()->GetFName()))
+		{
+			PropertySuper = Property->GetOwnerClass();
+		}
+	}
+
+	// Getter	
+	FString GetterName = FString::Printf(TEXT("Get_%s"), *PropertyName);
+	FString GeneratedGlue = GenerateWrapperFunctionDeclaration(ClassNameCPP, Class, GetterName);
+	GeneratedGlue += TEXT("\r\n{\r\n");
+	FString FunctionBody;
+	if (PropertySuper == NULL)
+	{
+		FunctionBody += FString::Printf(TEXT("\t%s\r\n"), *GenerateObjectDeclarationFromContext(ClassNameCPP, Class));
+		FunctionBody += FString::Printf(TEXT("\tstatic UProperty* Property = FindScriptPropertyHelper(%s::StaticClass(), TEXT(\"%s\"));\r\n"), *ClassNameCPP, *Property->GetName());
+		FunctionBody += FString::Printf(TEXT("\t%s PropertyValue;\r\n"), *GetPropertyTypeCPP(Property, CPPF_ArgumentOrReturnValue));
+		FunctionBody += TEXT("\tProperty->CopyCompleteValue(&PropertyValue, Property->ContainerPtrToValuePtr<void>(Obj));\r\n");
+		FunctionBody += FString::Printf(TEXT("\t%s\r\n"), *GenerateReturnValueHandler(ClassNameCPP, Class, NULL, Property, TEXT("PropertyValue")));
+	}
+	else
+	{
+		FunctionBody = FString::Printf(TEXT("\treturn %s_%s(InScriptContext);\r\n"), *PropertySuper->GetName(), *GetterName);
+	}
+	GeneratedGlue += FunctionBody;
+	GeneratedGlue += TEXT("}\r\n\r\n");
+	FunctionBody.Empty(FunctionBody.Len());
+
+	// Store the name of this getter as well as the name of the wrapper function
+	FPropertyAccessor Getter;
+	Getter.AccessorName = GetterName;
+	Getter.FunctionName = FString::Printf(TEXT("%s_%s"), *Class->GetName(), *GetterName);
+	auto& Exports = ClassExportedProperties.FindOrAdd(Class);
+	Exports.Add(Getter);
+
+	// Setter
+	FString SetterName = FString::Printf(TEXT("Set_%s"), *PropertyName);
+	GeneratedGlue += GenerateWrapperFunctionDeclaration(ClassNameCPP, Class, SetterName);
+	GeneratedGlue += TEXT("\r\n{\r\n");
+	if (PropertySuper == NULL)
+	{
+		FunctionBody += FString::Printf(TEXT("\t%s\r\n"), *GenerateObjectDeclarationFromContext(ClassNameCPP, Class));
+		FunctionBody += FString::Printf(TEXT("\tstatic UProperty* Property = FindScriptPropertyHelper(%s::StaticClass(), TEXT(\"%s\"));\r\n"), *ClassNameCPP, *Property->GetName());
+		FunctionBody += FString::Printf(TEXT("\t%s PropertyValue = %s;\r\n"), *GetPropertyTypeCPP(Property, CPPF_ArgumentOrReturnValue), *InitializeFunctionDispatchParam(NULL, Property, 0));
+		FunctionBody += TEXT("\tProperty->CopyCompleteValue(Property->ContainerPtrToValuePtr<void>(Obj), &PropertyValue);\r\n");
+		FunctionBody += TEXT("\treturn 0;\r\n");
+	}
+	else
+	{
+		FunctionBody = FString::Printf(TEXT("\treturn %s_%s(InScriptContext);\r\n"), *PropertySuper->GetName(), *SetterName);
+	}
+	GeneratedGlue += FunctionBody;	
+	GeneratedGlue += TEXT("}\r\n\r\n");
+
+	// Store the name of this setter as well as the name of the wrapper function
+	FPropertyAccessor Setter;
+	Setter.AccessorName = SetterName;
+	Setter.FunctionName = FString::Printf(TEXT("%s_%s"), *Class->GetName(), *SetterName);
+	Exports.Add(Setter);
+
 	return GeneratedGlue;
 }
 
@@ -292,7 +396,9 @@ FString FLuaScriptCodeGenerator::ExportAdditionalClassGlue(const FString& ClassN
 	{
 		GeneratedGlue += GenerateWrapperFunctionDeclaration(ClassNameCPP, Class, TEXT("New"));
 		GeneratedGlue += TEXT("\r\n{\r\n");
-		GeneratedGlue += FString::Printf(TEXT("\tUObject* Obj = NewObject<%s>();\r\n"), *ClassNameCPP);
+		GeneratedGlue += TEXT("\tUObject* Outer = (UObject*)lua_touserdata(InScriptContext, 1);\r\n");
+		GeneratedGlue += TEXT("\tFName Name = FName(luaL_checkstring(InScriptContext, 2));\r\n");
+		GeneratedGlue += FString::Printf(TEXT("\tUObject* Obj = NewNamedObject<%s>(Outer, Name);\r\n"), *ClassNameCPP);
 		GeneratedGlue += TEXT("\tif (Obj)\r\n\t{\r\n");
 		GeneratedGlue += TEXT("\t\tFScriptObjectReferencer::Get().AddObjectReference(Obj);\r\n");
 		GeneratedGlue += TEXT("\t}\r\n");
@@ -325,14 +431,22 @@ FString FLuaScriptCodeGenerator::ExportAdditionalClassGlue(const FString& ClassN
 		GeneratedGlue += FString::Printf(TEXT("\t{ \"Destroy\", %s_Destroy },\r\n"), *ClassName);
 		GeneratedGlue += FString::Printf(TEXT("\t{ \"Class\", %s_Class },\r\n"), *ClassName);
 	}
-	auto Exports = ClassExportedFunctions.Find(Class);
-	if (Exports)
+	auto FunctionExports = ClassExportedFunctions.Find(Class);
+	if (FunctionExports)
 	{
-		for (auto& FunctionName : *Exports)
+		for (auto& FunctionName : *FunctionExports)
 		{
 			GeneratedGlue += FString::Printf(TEXT("\t{ \"%s\", %s_%s },\r\n"), *FunctionName.ToString(), *ClassName, *FunctionName.ToString());
 		}
 	}
+	auto PropertyExports = ClassExportedProperties.Find(Class);
+	if (PropertyExports)
+	{
+		for (auto& Accessor : *PropertyExports)
+		{
+			GeneratedGlue += FString::Printf(TEXT("\t{ \"%s\", %s },\r\n"), *Accessor.AccessorName, *Accessor.FunctionName);
+		}
+	}
 	GeneratedGlue += TEXT("\t{ NULL, NULL }\r\n};\r\n\r\n");
 
 	return GeneratedGlue;
@@ -369,6 +483,7 @@ void FLuaScriptCodeGenerator::ExportClass(UClass* Class, const FString& SourceHe
 		}
 	}
 
+	// Export all functions
 	for (TFieldIterator<UFunction> FuncIt(Class /*, EFieldIteratorFlags::ExcludeSuper*/); FuncIt; ++FuncIt)
 	{
 		UFunction* Function = *FuncIt;
@@ -378,6 +493,18 @@ void FLuaScriptCodeGenerator::ExportClass(UClass* Class, const FString& SourceHe
 		}
 	}
 
+	// Export properties that are owned by this class
+	int32 PropertyIndex = 0;
+	for (TFieldIterator<UProperty> PropertyIt(Class /*, EFieldIteratorFlags::ExcludeSuper*/); PropertyIt; ++PropertyIt, ++PropertyIndex)
+	{
+		UProperty* Property = *PropertyIt;
+		if (CanExportProperty(ClassNameCPP, Class, Property))
+		{
+			UE_LOG(LogScriptGenerator, Log, TEXT("  %s %s"), *Property->GetClass()->GetName(), *Property->GetName());
+			GeneratedGlue += ExportProperty(ClassNameCPP, Class, Property, PropertyIndex);
+		}
+	}
+
 	GeneratedGlue += ExportAdditionalClassGlue(ClassNameCPP, Class);
 
 	SaveHeaderIfChanged(ClassGlueFilename, GeneratedGlue);
diff --git a/Engine/Plugins/Script/ScriptGeneratorPlugin/Source/ScriptGeneratorPlugin/Private/LuaScriptCodeGenerator.h b/Engine/Plugins/Script/ScriptGeneratorPlugin/Source/ScriptGeneratorPlugin/Private/LuaScriptCodeGenerator.h
--- a/Engine/Plugins/Script/ScriptGeneratorPlugin/Source/ScriptGeneratorPlugin/Private/LuaScriptCodeGenerator.h
+++ b/Engine/Plugins/Script/ScriptGeneratorPlugin/Source/ScriptGeneratorPlugin/Private/LuaScriptCodeGenerator.h
@@ -10,6 +10,12 @@ class FLuaScriptCodeGenerator : public FScriptCodeGeneratorBase
 {
 protected:
 
+	struct FPropertyAccessor
+	{
+		FString FunctionName;
+		FString AccessorName;
+	};
+
 	/** All generated script header filenames */
 	TArray<FString> AllScriptHeaders;
 	/** Source header filenames for all exported classes */
@@ -18,6 +24,8 @@ protected:
 	TArray<UClass*> LuaExportedClasses;
 	/** Functions exported for a class */
 	TMap<UClass*, TArray<FName> > ClassExportedFunctions;
+	/** Proprties exported for a class */
+	TMap<UClass*, TArray<FPropertyAccessor> > ClassExportedProperties;
 
 	/** Creats a 'glue' file that merges all generated script files */
 	void GlueAllGeneratedFiles();
@@ -25,23 +33,25 @@ protected:
 	/** Exports a wrapper function */
 	FString ExportFunction(const FString& ClassNameCPP, UClass* Class, UFunction* Function);
 	/** Exports a wrapper functions for properties */
-	FString ExportProperty(const FString& ClassNameCPP, UClass* Class, UProperty* Property);
+	FString ExportProperty(const FString& ClassNameCPP, UClass* Class, UProperty* Property, int32 PropertyIndex);
 	/** Exports additional class glue code (like ctor/dtot) */
 	FString ExportAdditionalClassGlue(const FString& ClassNameCPP, UClass* Class);
 	/** Generates wrapper function declaration */
 	FString GenerateWrapperFunctionDeclaration(const FString& ClassNameCPP, UClass* Class, UFunction* Function);
 	/** Generates wrapper function declaration */
 	FString GenerateWrapperFunctionDeclaration(const FString& ClassNameCPP, UClass* Class, const FString& FunctionName);
-	/** Generates function parameter declaration */
-	FString GenerateFunctionParamDeclaration(const FString& ClassNameCPP, UClass* Class, UFunction* Function, UProperty* Param, int32 ParamIndex);
 	/** Generates code responsible for getting the object pointer from script context */
 	FString GenerateObjectDeclarationFromContext(const FString& ClassNameCPP, UClass* Class);
 	/** Handles the wrapped function's return value */
-	FString GenerateReturnValueHandler(const FString& ClassNameCPP, UClass* Class, UFunction* Function, UProperty* ReturnValue);
+	FString GenerateReturnValueHandler(const FString& ClassNameCPP, UClass* Class, UFunction* Function, UProperty* ReturnValue, const FString& ReturnValueName);
+	/** Check if a property type is supported */
+	bool IsPropertyTypeSupported(UProperty* Property) const;
 
 	// FScriptCodeGeneratorBase interface
 	virtual bool CanExportClass(UClass* Class) OVERRIDE;
 	virtual bool CanExportFunction(const FString& ClassNameCPP, UClass* Class, UFunction* Function) OVERRIDE;
+	virtual bool CanExportProperty(const FString& ClassNameCPP, UClass* Class, UProperty* Property) OVERRIDE;
+	virtual FString InitializeFunctionDispatchParam(UFunction* Function, UProperty* Param, int32 ParamIndex) OVERRIDE;
 
 public:
 
diff --git a/Engine/Plugins/Script/ScriptGeneratorPlugin/Source/ScriptGeneratorPlugin/Private/ScriptCodeGeneratorBase.cpp b/Engine/Plugins/Script/ScriptGeneratorPlugin/Source/ScriptGeneratorPlugin/Private/ScriptCodeGeneratorBase.cpp
--- a/Engine/Plugins/Script/ScriptGeneratorPlugin/Source/ScriptGeneratorPlugin/Private/ScriptCodeGeneratorBase.cpp
+++ b/Engine/Plugins/Script/ScriptGeneratorPlugin/Source/ScriptGeneratorPlugin/Private/ScriptCodeGeneratorBase.cpp
@@ -98,6 +98,54 @@ FString FScriptCodeGeneratorBase::GetPropertyTypeCPP(UProperty* Property, uint32
 	return PropertyType;
 }
 
+FString FScriptCodeGeneratorBase::GenerateFunctionDispatch(UFunction* Function)
+{
+	FString Params;
+	
+	const bool bHasParamsOrReturnValue = (Function->Children != NULL);
+	if (bHasParamsOrReturnValue)
+	{
+		Params += TEXT("\tstruct FDispatchParams\r\n\t{\r\n");
+
+		for (TFieldIterator<UProperty> ParamIt(Function); ParamIt; ++ParamIt)
+		{
+			UProperty* Param = *ParamIt;
+			Params += FString::Printf(TEXT("\t\t%s %s;\r\n"), *GetPropertyTypeCPP(Param, CPPF_ArgumentOrReturnValue), *Param->GetName());
+		}
+		Params += TEXT("\t} Params;\r\n");
+		int32 ParamIndex = 0;
+		for (TFieldIterator<UProperty> ParamIt(Function); ParamIt; ++ParamIt, ++ParamIndex)
+		{
+			UProperty* Param = *ParamIt;
+			Params += FString::Printf(TEXT("\tParams.%s = %s;\r\n"), *Param->GetName(), *InitializeFunctionDispatchParam(Function, Param, ParamIndex));
+		}
+	}
+	Params += FString::Printf(TEXT("\tstatic UFunction* Function = Obj->FindFunctionChecked(TEXT(\"%s\"));\r\n"), *Function->GetName());
+	if (bHasParamsOrReturnValue)
+	{
+		Params += TEXT("\tcheck(Function->ParmsSize == sizeof(FDispatchParams));\r\n");
+		Params += TEXT("\tObj->ProcessEvent(Function, &Params);\r\n");
+	}
+	else
+	{
+		Params += TEXT("\tObj->ProcessEvent(Function, NULL);\r\n");
+	}	
+
+	return Params;
+}
+
+FString FScriptCodeGeneratorBase::InitializeFunctionDispatchParam(UFunction* Function, UProperty* Param, int32 ParamIndex)
+{
+	if (Param->IsA(UObjectPropertyBase::StaticClass()) || Param->IsA(UClassProperty::StaticClass()))
+	{
+		return TEXT("NULL");
+	}
+	else
+	{
+		return FString::Printf(TEXT("%s()"), *GetPropertyTypeCPP(Param, CPPF_ArgumentOrReturnValue));
+	}	
+}
+
 FString FScriptCodeGeneratorBase::GetScriptHeaderForClass(UClass* Class)
 {
 	return GeneratedCodePath / (Class->GetName() + TEXT(".script.h"));
@@ -129,15 +177,7 @@ bool FScriptCodeGeneratorBase::CanExportClass(UClass* Class)
 bool FScriptCodeGeneratorBase::CanExportFunction(const FString& ClassNameCPP, UClass* Class, UFunction* Function)
 {
 	// We don't support delegates and non-public functions
-	if ((Function->FunctionFlags & FUNC_Delegate) ||
-		!(Function->FunctionFlags & FUNC_Public))
-	{
-		return false;
-	}
-
-	// Function must be DLL exported
-	if ((Class->ClassFlags & CLASS_RequiredAPI) == 0 &&
-		(Function->FunctionFlags & FUNC_RequiredAPI) == 0)
+	if ((Function->FunctionFlags & FUNC_Delegate))
 	{
 		return false;
 	}
@@ -191,30 +231,3 @@ bool FScriptCodeGeneratorBase::CanExportProperty(const FString& ClassNameCPP, UC
 
 	return true;
 }
-
-int32 FScriptCodeGeneratorBase::GenerateExportedPropertyTable(const FString& ClassNameCPP, UClass* Class, FString& OutFunction, TArray<UProperty*>& OutExportedProperties)
-{
-	OutExportedProperties.Empty();
-	OutFunction = FString::Printf(TEXT("UProperty* %s_GetExportedProperty(int32 PropertyIndex)\r\n{\r\n"), *Class->GetName());
-	OutFunction += FString::Printf(TEXT("\tstatic UClass* Class = %s::StaticClass();\r\n"), *ClassNameCPP);
-	OutFunction += TEXT("\tstatic UProperty* PropertyTable[] =\r\n\t{\r\n");
-	
-	for (TFieldIterator<UProperty> PropertyIt(Class, EFieldIteratorFlags::ExcludeSuper); PropertyIt; ++PropertyIt)
-	{
-		UProperty* Property = *PropertyIt;
-		if (CanExportProperty(ClassNameCPP, Class, Property))
-		{
-			if (OutExportedProperties.Num() > 0)
-			{
-				OutFunction += TEXT(",\r\n");
-			}
-			OutFunction += FString::Printf(TEXT("\t\tFindScriptPropertyHelper(Class, TEXT(\"%s\"))"), *Property->GetName());
-			OutExportedProperties.Add(Property);
-		}
-	}
-
-	OutFunction += TEXT("\r\n\t};\r\n\tcheck(PropertyIndex >= 0 && PropertyIndex < ARRAY_COUNT(PropertyTable));\r\n");
-	OutFunction += TEXT("\treturn PropertyTable[PropertyIndex];\r\n}\r\n\r\n");
-
-	return OutExportedProperties.Num();
-}
\ No newline at end of file
diff --git a/Engine/Plugins/Script/ScriptGeneratorPlugin/Source/ScriptGeneratorPlugin/Private/ScriptCodeGeneratorBase.h b/Engine/Plugins/Script/ScriptGeneratorPlugin/Source/ScriptGeneratorPlugin/Private/ScriptCodeGeneratorBase.h
--- a/Engine/Plugins/Script/ScriptGeneratorPlugin/Source/ScriptGeneratorPlugin/Private/ScriptCodeGeneratorBase.h
+++ b/Engine/Plugins/Script/ScriptGeneratorPlugin/Source/ScriptGeneratorPlugin/Private/ScriptCodeGeneratorBase.h
@@ -34,6 +34,10 @@ protected:
 	FString GetClassNameCPP(UClass* Class);
 	/** Gets C++ friendly property type name. */
 	FString GetPropertyTypeCPP(UProperty* Property, uint32 PortFlags = 0);
+
+	virtual FString GenerateFunctionDispatch(UFunction* Function);
+	virtual FString InitializeFunctionDispatchParam(UFunction* Function, UProperty* Param, int32 ParamIndex);
+
 	/** Generates a script glue header filename for the specified class */
 	virtual FString GetScriptHeaderForClass(UClass* Class);
 	/** Returns true if the specified class can be exported */
@@ -42,8 +46,6 @@ protected:
 	virtual bool CanExportFunction(const FString& ClassNameCPP, UClass* Class, UFunction* Function);
 	/** Returns true if the specified property can be exported */
 	virtual bool CanExportProperty(const FString& ClassNameCPP, UClass* Class, UProperty* Property);
-	/** Generates a lookup table for exported proeprties. */
-	virtual int32 GenerateExportedPropertyTable(const FString& ClassNameCPP, UClass* Class, FString& OutFunction, TArray<UProperty*>& OutExportedProperties);
 
 public:
 
diff --git a/Engine/Plugins/Script/ScriptPlugin/Source/ScriptPlugin/Classes/ScriptAsset.h b/Engine/Plugins/Script/ScriptPlugin/Source/ScriptPlugin/Classes/ScriptAsset.h
--- a/Engine/Plugins/Script/ScriptPlugin/Source/ScriptPlugin/Classes/ScriptAsset.h
+++ b/Engine/Plugins/Script/ScriptPlugin/Source/ScriptPlugin/Classes/ScriptAsset.h
@@ -30,6 +30,9 @@ class UScriptAsset : public UDataAsset
 #if WITH_EDITOR
 	virtual void PostEditChangeProperty(struct FPropertyChangedEvent& PropertyChangedEvent);
 
+	/** Re-load script if it has changed and compile */
+	virtual void UpdateAndCompile();
+
 	/** Compile script to bytecode */
 	virtual void Compile();
 #endif
diff --git a/Engine/Plugins/Script/ScriptPlugin/Source/ScriptPlugin/Classes/ScriptComponent.h b/Engine/Plugins/Script/ScriptPlugin/Source/ScriptPlugin/Classes/ScriptComponent.h
--- a/Engine/Plugins/Script/ScriptPlugin/Source/ScriptPlugin/Classes/ScriptComponent.h
+++ b/Engine/Plugins/Script/ScriptPlugin/Source/ScriptPlugin/Classes/ScriptComponent.h
@@ -9,6 +9,7 @@ class FScriptContextBase
 {
 public:
 	virtual bool Initialize(class UScriptComponent* Owner) = 0;
+	virtual void BeginPlay() = 0;
 	virtual void Tick(float DeltaTime) = 0;
 	virtual void Destroy() = 0;
 	virtual bool CanTick() = 0;
@@ -29,6 +30,7 @@ public:
 
 	// Begin UActorComponent interface.
 	virtual void OnRegister() OVERRIDE;
+	virtual void InitializeComponent() OVERRIDE;
 	virtual void TickComponent(float DeltaTime, enum ELevelTick TickType, FActorComponentTickFunction *ThisTickFunction) OVERRIDE;
 	virtual void OnUnregister() OVERRIDE;
 	// Begin UActorComponent interface.
diff --git a/Engine/Plugins/Script/ScriptPlugin/Source/ScriptPlugin/Classes/ScriptTestActor.h b/Engine/Plugins/Script/ScriptPlugin/Source/ScriptPlugin/Classes/ScriptTestActor.h
new file mode 100644
--- /dev/null
+++ b/Engine/Plugins/Script/ScriptPlugin/Source/ScriptPlugin/Classes/ScriptTestActor.h
@@ -0,0 +1,22 @@
+// Copyright 1998-2014 Epic Games, Inc. All Rights Reserved.
+#pragma once
+#include "ScriptTestActor.generated.h"
+
+/** Definition of a specific ability that is applied to a character. Exists as part of a Trait. */
+UCLASS(BlueprintType, MinimalAPI)
+class AScriptTestActor : public AActor
+{
+	GENERATED_UCLASS_BODY()
+
+	UPROPERTY(Category = Test, EditAnywhere, BlueprintReadWrite)
+	FString TestString;
+
+	UPROPERTY(Category = Test, EditAnywhere, BlueprintReadWrite)
+	float TestValue;
+
+	UPROPERTY(Category = Test, EditAnywhere, BlueprintReadWrite)
+	bool TestBool;
+
+	UFUNCTION()
+	float TestFunction(float InValue, float InFactor, bool bMultiply);
+};
diff --git a/Engine/Plugins/Script/ScriptPlugin/Source/ScriptPlugin/Private/LuaIntegration.cpp b/Engine/Plugins/Script/ScriptPlugin/Source/ScriptPlugin/Private/LuaIntegration.cpp
--- a/Engine/Plugins/Script/ScriptPlugin/Source/ScriptPlugin/Private/LuaIntegration.cpp
+++ b/Engine/Plugins/Script/ScriptPlugin/Source/ScriptPlugin/Private/LuaIntegration.cpp
@@ -2,6 +2,7 @@
 
 #include "ScriptPluginPrivatePCH.h"
 #include "LuaIntegration.h"
+#include "ScriptObjectReferencer.h"
 
 #if WITH_LUA
 
@@ -94,13 +95,15 @@ struct FLuaUObject
 	* LoadObject
 	*
 	* @param Class
+	* @param Outer
 	* @param LongPackageName
 	*/
 	static int32 LuaStaticLoadObject(lua_State* LuaState)
 	{
 		UClass* Class = (UClass*)lua_touserdata(LuaState, 1);
-		FString LongPackageName = ANSI_TO_TCHAR(luaL_checkstring(LuaState, 2));
-		UObject* Result = StaticLoadObject(Class, NULL, *LongPackageName);
+		UObject* Outer = (UObject*)lua_touserdata(LuaState, 2);
+		FString LongPackageName = ANSI_TO_TCHAR(luaL_checkstring(LuaState, 3));
+		UObject* Result = StaticLoadObject(Class, Outer, *LongPackageName);
 		lua_pushlightuserdata(LuaState, Result);
 		return 1;
 	}
@@ -120,6 +123,108 @@ struct FLuaUObject
 		return 1;
 	}
 
+	/**
+	* GetCurrentLevel
+	*
+	* @param 'this'
+	*/
+	static int32 LuaGetCurrentLevel(lua_State* LuaState)
+	{
+		ULevel* Level = NULL;
+		UObject* This = (UObject*)lua_touserdata(LuaState, 1);
+		UActorComponent* ThisComponent = Cast<UActorComponent>(This);
+		if (ThisComponent && ThisComponent->GetOwner())
+		{
+			Level = ThisComponent->GetOwner()->GetLevel();
+		}
+		lua_pushlightuserdata(LuaState, Level);
+		return 1;
+	}
+
+	/**
+	* LuaGetName
+	*
+	* @param Object
+	*/
+	static int32 LuaGetName(lua_State* LuaState)
+	{
+		UObject* Object = (UObject*)lua_touserdata(LuaState, 1);
+		lua_pushstring(LuaState, TCHAR_TO_ANSI(*Object->GetName()));
+		return 1;
+	}
+
+	/**
+	* LuaGetFullName
+	*
+	* @param Object
+	*/
+	static int32 LuaGetFullName(lua_State* LuaState)
+	{
+		UObject* Object = (UObject*)lua_touserdata(LuaState, 1);
+		lua_pushstring(LuaState, TCHAR_TO_ANSI(*Object->GetFullName()));
+		return 1;
+	}
+
+	/**
+	* LuaGetPathName
+	*
+	* @param Object
+	*/
+	static int32 LuaGetPathName(lua_State* LuaState)
+	{
+		UObject* Object = (UObject*)lua_touserdata(LuaState, 1);
+		lua_pushstring(LuaState, TCHAR_TO_ANSI(*Object->GetPathName()));
+		return 1;
+	}
+
+	/**
+	* LuaGetOuter
+	*
+	* @param Object
+	*/
+	static int32 LuaGetOuter(lua_State* LuaState)
+	{
+		UObject* Object = (UObject*)lua_touserdata(LuaState, 1);
+		lua_pushlightuserdata(LuaState, Object->GetOuter());
+		return 1;
+	}
+
+	/**
+	* LuaIsValid
+	*
+	* @param Object
+	*/
+	static int32 LuaIsValid(lua_State* LuaState)
+	{
+		UObject* Object = (UObject*)lua_touserdata(LuaState, 1);
+		lua_pushboolean(LuaState, IsValid(Object));
+		return 1;
+	}
+
+	/**
+	* LuaAddReferencedObject
+	*
+	* @param Object
+	*/
+	static int32 LuaAddReferencedObject(lua_State* LuaState)
+	{
+		UObject* Object = (UObject*)lua_touserdata(LuaState, 1);
+		FScriptObjectReferencer::Get().AddObjectReference(Object);
+		return 0;
+	}
+
+	/**
+	* LuaAddReferencedObject
+	*
+	* @param Object
+	*/
+	static int32 LuaRemoveReferencedObject(lua_State* LuaState)
+	{
+		UObject* Object = (UObject*)lua_touserdata(LuaState, 1);
+		FScriptObjectReferencer::Get().RemoveObjectReference(Object);
+		return 0;
+	}
+
 	// Library
 	static luaL_Reg UObject_Lib[];
 };
@@ -130,6 +235,14 @@ luaL_Reg FLuaUObject::UObject_Lib[] =
 	{ "FindObject", FLuaUObject::LuaStaticFindObject },
 	{ "FindObjectFast", FLuaUObject::LuaStaticFindObjectFast },
 	{ "LoadObject", FLuaUObject::LuaStaticLoadObject },
+	{ "GetCurrentLevel", FLuaUObject::LuaGetCurrentLevel },
+	{ "GetName", FLuaUObject::LuaGetName },
+	{ "GetFullName", FLuaUObject::LuaGetFullName },
+	{ "GetPathName", FLuaUObject::LuaGetPathName },
+	{ "GetOuter", FLuaUObject::LuaGetOuter },
+	{ "IsValid", FLuaUObject::LuaIsValid },
+	{ "AddReferencedObject", FLuaUObject::LuaAddReferencedObject },
+	{ "RemoveReferencedObject", FLuaUObject::LuaRemoveReferencedObject },
 	{ NULL, NULL }
 };
 
@@ -388,6 +501,7 @@ static void LuaRegisterUnrealUtilities(lua_State* LuaState)
 FLuaContext::FLuaContext()
 : bHasTick(false)
 , bHasDestroy(false)
+, bHasBeginPlay(false)
 , LuaState(NULL)
 {
 };
@@ -411,6 +525,7 @@ bool FLuaContext::Initialize(UScriptComponent* Owner)
 		{
 			bHasTick = FLuaUtils::DoesFunctionExist(LuaState, "Tick");
 			bHasDestroy = FLuaUtils::DoesFunctionExist(LuaState, "Destroy");
+			bHasBeginPlay = FLuaUtils::DoesFunctionExist(LuaState, "BeginPlay");
 		}
 	}
 	else
@@ -420,6 +535,15 @@ bool FLuaContext::Initialize(UScriptComponent* Owner)
 	return bResult;
 }
 
+void FLuaContext::BeginPlay()
+{
+	check(LuaState);
+	if (bHasBeginPlay)
+	{
+		FLuaUtils::CallFunction(LuaState, "BeginPlay");
+	}
+}
+
 void FLuaContext::Tick(float DeltaTime)
 {
 	check(LuaState && bHasTick);
diff --git a/Engine/Plugins/Script/ScriptPlugin/Source/ScriptPlugin/Private/LuaIntegration.h b/Engine/Plugins/Script/ScriptPlugin/Source/ScriptPlugin/Private/LuaIntegration.h
--- a/Engine/Plugins/Script/ScriptPlugin/Source/ScriptPlugin/Private/LuaIntegration.h
+++ b/Engine/Plugins/Script/ScriptPlugin/Source/ScriptPlugin/Private/LuaIntegration.h
@@ -14,6 +14,7 @@ class FLuaContext : public FScriptContextBase
 {
 	bool bHasTick;
 	bool bHasDestroy;
+	bool bHasBeginPlay;
 	lua_State* LuaState;
 
 public:
@@ -22,6 +23,7 @@ public:
 
 	// FScriptContextBase interface
 	virtual bool Initialize(UScriptComponent* Owner) OVERRIDE;
+	virtual void BeginPlay() OVERRIDE;
 	virtual void Tick(float DeltaTime) OVERRIDE;
 	virtual void Destroy() OVERRIDE;
 	virtual bool CanTick() OVERRIDE;
diff --git a/Engine/Plugins/Script/ScriptPlugin/Source/ScriptPlugin/Private/ScriptAsset.cpp b/Engine/Plugins/Script/ScriptPlugin/Source/ScriptPlugin/Private/ScriptAsset.cpp
--- a/Engine/Plugins/Script/ScriptPlugin/Source/ScriptPlugin/Private/ScriptAsset.cpp
+++ b/Engine/Plugins/Script/ScriptPlugin/Source/ScriptPlugin/Private/ScriptAsset.cpp
@@ -13,10 +13,7 @@ void UScriptAsset::PostLoad()
 	Super::PostLoad();
 
 #if WITH_EDITOR
-	if (!SourceCode.IsEmpty() && !ByteCode.Num())
-	{
-		Compile();
-	}
+	UpdateAndCompile();
 #endif
 }
 
@@ -24,14 +21,34 @@ void UScriptAsset::PostLoad()
 void UScriptAsset::PostEditChangeProperty(struct FPropertyChangedEvent& PropertyChangedEvent)
 {
 	Super::PostEditChangeProperty(PropertyChangedEvent);
-	Compile();
+	UpdateAndCompile();
 }
 
-void UScriptAsset::Compile()
+void UScriptAsset::UpdateAndCompile()
 {
-	if (GIsEditor)
+	FDateTime OldTimeStamp;
+	bool bCodeDirty = !FDateTime::Parse(SourceFileTimestamp, OldTimeStamp);
+	FDateTime TimeStamp = IFileManager::Get().GetTimeStamp(*SourceFilePath);
+	if (bCodeDirty || TimeStamp > OldTimeStamp)
+	{
+		FString NewScript;
+		if (FFileHelper::LoadFileToString(NewScript, *SourceFilePath))
+		{
+			SourceCode = NewScript;
+			SourceFileTimestamp = TimeStamp.ToString();
+			bCodeDirty = true;			
+		}
+	}
+
+	if (bCodeDirty)
 	{
+		Compile();
 		MarkPackageDirty();
-	}	
+	}
 }
+
+void UScriptAsset::Compile()
+{
+}
+
 #endif
\ No newline at end of file
diff --git a/Engine/Plugins/Script/ScriptPlugin/Source/ScriptPlugin/Private/ScriptComponent.cpp b/Engine/Plugins/Script/ScriptPlugin/Source/ScriptPlugin/Private/ScriptComponent.cpp
--- a/Engine/Plugins/Script/ScriptPlugin/Source/ScriptPlugin/Private/ScriptComponent.cpp
+++ b/Engine/Plugins/Script/ScriptPlugin/Source/ScriptPlugin/Private/ScriptComponent.cpp
@@ -10,20 +10,25 @@ UScriptComponent::UScriptComponent(const FPostConstructInitializeProperties& PCI
 	PrimaryComponentTick.bCanEverTick = true;
 	bTickInEditor = false;
 	bAutoActivate = true;
+	bWantsInitializeComponent = true;
 
 	Context = NULL;
 }
 
-
 void UScriptComponent::OnRegister()
 {
 	Super::OnRegister();
 	if (Script && GetWorld() && GetWorld()->WorldType != EWorldType::Editor)
 	{
-		// @todo Create context here
+#if WITH_EDITOR
+		Script->UpdateAndCompile();
+#endif
 #if WITH_LUA
 		Context = new FLuaContext();
+#else
+		// @todo Create context here
 #endif
+
 		if (Context)
 		{
 			bool bDoNotTick = true;
@@ -32,7 +37,7 @@ void UScriptComponent::OnRegister()
 				bDoNotTick = !Context->CanTick();
 			}
 			if (bDoNotTick)
-			{				
+			{
 				bAutoActivate = false;
 				PrimaryComponentTick.bCanEverTick = false;
 			}
@@ -40,6 +45,15 @@ void UScriptComponent::OnRegister()
 	}
 }
 
+void UScriptComponent::InitializeComponent()
+{
+	Super::InitializeComponent();
+	if (Context)
+	{
+		Context->BeginPlay();
+	}
+}
+
 void UScriptComponent::TickComponent(float DeltaTime, enum ELevelTick TickType, FActorComponentTickFunction *ThisTickFunction)
 {
 	Super::TickComponent(DeltaTime, TickType, ThisTickFunction);
diff --git a/Engine/Plugins/Script/ScriptPlugin/Source/ScriptPlugin/Private/ScriptTestActor.cpp b/Engine/Plugins/Script/ScriptPlugin/Source/ScriptPlugin/Private/ScriptTestActor.cpp
new file mode 100644
--- /dev/null
+++ b/Engine/Plugins/Script/ScriptPlugin/Source/ScriptPlugin/Private/ScriptTestActor.cpp
@@ -0,0 +1,21 @@
+// Copyright 1998-2014 Epic Games, Inc. All Rights Reserved. 
+#include "ScriptPluginPrivatePCH.h"
+
+//////////////////////////////////////////////////////////////////////////
+
+AScriptTestActor::AScriptTestActor(const FPostConstructInitializeProperties& PCIP)
+	: Super( PCIP )
+{
+}
+
+float AScriptTestActor::TestFunction(float InValue, float InFactor, bool bMultiply)
+{
+	if (bMultiply)
+	{
+		return InValue * InFactor;
+	}
+	else
+	{
+		return InValue / InFactor;
+	}
+}
\ No newline at end of file
diff --git a/Engine/Programs/UnrealHeaderTool/Config/DefaultEngine.ini b/Engine/Programs/UnrealHeaderTool/Config/DefaultEngine.ini
--- a/Engine/Programs/UnrealHeaderTool/Config/DefaultEngine.ini
+++ b/Engine/Programs/UnrealHeaderTool/Config/DefaultEngine.ini
@@ -3,4 +3,5 @@
 +ScriptPlugins="ScriptGeneratorPlugin"
 
 +ScriptSupportedModules=/Script/CoreUObject
-+ScriptSupportedModules=/Script/Engine
\ No newline at end of file
++ScriptSupportedModules=/Script/Engine
++ScriptSupportedModules=/Script/ScriptPlugin
\ No newline at end of file
