ScriptPlugin: fixing inconsistent script glue generation depending if headers were generated after clean or not.

diff --git a/Engine/Plugins/Script/ScriptGeneratorPlugin/Source/ScriptGeneratorPlugin/Private/GenericScriptCodeGenerator.cpp b/Engine/Plugins/Script/ScriptGeneratorPlugin/Source/ScriptGeneratorPlugin/Private/GenericScriptCodeGenerator.cpp
--- a/Engine/Plugins/Script/ScriptGeneratorPlugin/Source/ScriptGeneratorPlugin/Private/GenericScriptCodeGenerator.cpp
+++ b/Engine/Plugins/Script/ScriptGeneratorPlugin/Source/ScriptGeneratorPlugin/Private/GenericScriptCodeGenerator.cpp
@@ -231,18 +231,6 @@ void FGenericScriptCodeGenerator::ExportClass(UClass* Class, const FString& Sour
 	const FString ClassNameCPP = GetClassNameCPP(Class);
 	FString GeneratedGlue(TEXT("#pragma once\r\n\r\n"));		
 
-	for (UClass* Super = Class->GetSuperClass(); Super != NULL; Super = Super->GetSuperClass())
-	{
-		FString SuperScriptHeader = GetScriptHeaderForClass(Super);
-		if (FPaths::FileExists(SuperScriptHeader))
-		{
-			// Re-base to make sure we're including the right files on a remote machine
-			FString NewFilename(RebaseToBuildPath(SuperScriptHeader));
-			GeneratedGlue += FString::Printf(TEXT("#include \"%s\"\r\n\r\n"), *NewFilename);
-			break;
-		}
-	}
-
 	for (TFieldIterator<UFunction> FuncIt(Class /*, EFieldIteratorFlags::ExcludeSuper*/); FuncIt; ++FuncIt)
 	{
 		UFunction* Function = *FuncIt;
diff --git a/Engine/Plugins/Script/ScriptGeneratorPlugin/Source/ScriptGeneratorPlugin/Private/LuaScriptCodeGenerator.cpp b/Engine/Plugins/Script/ScriptGeneratorPlugin/Source/ScriptGeneratorPlugin/Private/LuaScriptCodeGenerator.cpp
--- a/Engine/Plugins/Script/ScriptGeneratorPlugin/Source/ScriptGeneratorPlugin/Private/LuaScriptCodeGenerator.cpp
+++ b/Engine/Plugins/Script/ScriptGeneratorPlugin/Source/ScriptGeneratorPlugin/Private/LuaScriptCodeGenerator.cpp
@@ -471,18 +471,6 @@ void FLuaScriptCodeGenerator::ExportClass(UClass* Class, const FString& SourceHe
 	const FString ClassNameCPP = GetClassNameCPP(Class);
 	FString GeneratedGlue(TEXT("#pragma once\r\n\r\n"));		
 
-	for (UClass* Super = Class->GetSuperClass(); Super != NULL; Super = Super->GetSuperClass())
-	{
-		FString SuperScriptHeader = GetScriptHeaderForClass(Super);
-		if (FPaths::FileExists(SuperScriptHeader))
-		{
-			// Re-base to make sure we're including the right files on a remote machine
-			FString NewFilename(RebaseToBuildPath(SuperScriptHeader));
-			GeneratedGlue += FString::Printf(TEXT("#include \"%s\"\r\n\r\n"), *NewFilename);
-			break;
-		}
-	}
-
 	// Export all functions
 	for (TFieldIterator<UFunction> FuncIt(Class /*, EFieldIteratorFlags::ExcludeSuper*/); FuncIt; ++FuncIt)
 	{
