ScriptPlugin: Don't export virtual function overrides for the generic plugin.

diff --git a/Engine/Plugins/Script/ScriptGeneratorPlugin/Source/ScriptGeneratorPlugin/Private/GenericScriptCodeGenerator.cpp b/Engine/Plugins/Script/ScriptGeneratorPlugin/Source/ScriptGeneratorPlugin/Private/GenericScriptCodeGenerator.cpp
--- a/Engine/Plugins/Script/ScriptGeneratorPlugin/Source/ScriptGeneratorPlugin/Private/GenericScriptCodeGenerator.cpp
+++ b/Engine/Plugins/Script/ScriptGeneratorPlugin/Source/ScriptGeneratorPlugin/Private/GenericScriptCodeGenerator.cpp
@@ -231,7 +231,7 @@ void FGenericScriptCodeGenerator::ExportClass(UClass* Class, const FString& Sour
 	const FString ClassNameCPP = GetClassNameCPP(Class);
 	FString GeneratedGlue(TEXT("#pragma once\r\n\r\n"));		
 
-	for (TFieldIterator<UFunction> FuncIt(Class /*, EFieldIteratorFlags::ExcludeSuper*/); FuncIt; ++FuncIt)
+	for (TFieldIterator<UFunction> FuncIt(Class, EFieldIteratorFlags::ExcludeSuper); FuncIt; ++FuncIt)
 	{
 		UFunction* Function = *FuncIt;
 		if (CanExportFunction(ClassNameCPP, Class, Function))
