ScriptPlugins: generated glue files will now be including using relative paths

diff --git a/Engine/Plugins/Script/ScriptGeneratorPlugin/Source/ScriptGeneratorPlugin/Private/GenericScriptCodeGenerator.cpp b/Engine/Plugins/Script/ScriptGeneratorPlugin/Source/ScriptGeneratorPlugin/Private/GenericScriptCodeGenerator.cpp
--- a/Engine/Plugins/Script/ScriptGeneratorPlugin/Source/ScriptGeneratorPlugin/Private/GenericScriptCodeGenerator.cpp
+++ b/Engine/Plugins/Script/ScriptGeneratorPlugin/Source/ScriptGeneratorPlugin/Private/GenericScriptCodeGenerator.cpp
@@ -2,8 +2,8 @@
 #include "ScriptGeneratorPluginPrivatePCH.h"
 #include "GenericScriptCodeGenerator.h"
 
-FGenericScriptCodeGenerator::FGenericScriptCodeGenerator(const FString& RootLocalPath, const FString& RootBuildPath, const FString& OutputDirectory)
-: FScriptCodeGeneratorBase(RootLocalPath, RootBuildPath, OutputDirectory)
+FGenericScriptCodeGenerator::FGenericScriptCodeGenerator(const FString& RootLocalPath, const FString& RootBuildPath, const FString& OutputDirectory, const FString& InIncludeBase)
+: FScriptCodeGeneratorBase(RootLocalPath, RootBuildPath, OutputDirectory, InIncludeBase)
 {
 
 }
@@ -281,7 +281,7 @@ void FGenericScriptCodeGenerator::GlueAllGeneratedFiles()
 	for (auto& HeaderFilename : AllScriptHeaders)
 	{
 		// Re-base to make sure we're including the right files on a remote machine
-		FString NewFilename(RebaseToBuildPath(HeaderFilename));
+		FString NewFilename(FPaths::GetCleanFilename(HeaderFilename));
 		LibGlue += FString::Printf(TEXT("#include \"%s\"\r\n"), *NewFilename);
 	}
 
diff --git a/Engine/Plugins/Script/ScriptGeneratorPlugin/Source/ScriptGeneratorPlugin/Private/GenericScriptCodeGenerator.h b/Engine/Plugins/Script/ScriptGeneratorPlugin/Source/ScriptGeneratorPlugin/Private/GenericScriptCodeGenerator.h
--- a/Engine/Plugins/Script/ScriptGeneratorPlugin/Source/ScriptGeneratorPlugin/Private/GenericScriptCodeGenerator.h
+++ b/Engine/Plugins/Script/ScriptGeneratorPlugin/Source/ScriptGeneratorPlugin/Private/GenericScriptCodeGenerator.h
@@ -41,7 +41,7 @@ protected:
 
 public:
 
-	FGenericScriptCodeGenerator(const FString& RootLocalPath, const FString& RootBuildPath, const FString& OutputDirectory);
+	FGenericScriptCodeGenerator(const FString& RootLocalPath, const FString& RootBuildPath, const FString& OutputDirectory, const FString& InIncludeBase);
 
 	// FScriptCodeGeneratorBase interface
 	virtual void ExportClass(UClass* Class, const FString& SourceHeaderFilename, const FString& GeneratedHeaderFilename, bool bHasChanged) override;
diff --git a/Engine/Plugins/Script/ScriptGeneratorPlugin/Source/ScriptGeneratorPlugin/Private/LuaScriptCodeGenerator.cpp b/Engine/Plugins/Script/ScriptGeneratorPlugin/Source/ScriptGeneratorPlugin/Private/LuaScriptCodeGenerator.cpp
--- a/Engine/Plugins/Script/ScriptGeneratorPlugin/Source/ScriptGeneratorPlugin/Private/LuaScriptCodeGenerator.cpp
+++ b/Engine/Plugins/Script/ScriptGeneratorPlugin/Source/ScriptGeneratorPlugin/Private/LuaScriptCodeGenerator.cpp
@@ -8,8 +8,8 @@ static FName Name_Vector4("Vector4");
 static FName Name_Quat("Quat");
 static FName Name_Transform("Transform");
 
-FLuaScriptCodeGenerator::FLuaScriptCodeGenerator(const FString& RootLocalPath, const FString& RootBuildPath, const FString& OutputDirectory)
-: FScriptCodeGeneratorBase(RootLocalPath, RootBuildPath, OutputDirectory)
+FLuaScriptCodeGenerator::FLuaScriptCodeGenerator(const FString& RootLocalPath, const FString& RootBuildPath, const FString& OutputDirectory, const FString& InIncludeBase)
+: FScriptCodeGeneratorBase(RootLocalPath, RootBuildPath, OutputDirectory, InIncludeBase)
 {
 
 }
@@ -522,7 +522,7 @@ void FLuaScriptCodeGenerator::GlueAllGeneratedFiles()
 	for (auto& HeaderFilename : AllScriptHeaders)
 	{
 		// Re-base to make sure we're including the right files on a remote machine
-		FString NewFilename(RebaseToBuildPath(HeaderFilename));
+		FString NewFilename(FPaths::GetCleanFilename(HeaderFilename));
 		LibGlue += FString::Printf(TEXT("#include \"%s\"\r\n"), *NewFilename);
 	}
 
diff --git a/Engine/Plugins/Script/ScriptGeneratorPlugin/Source/ScriptGeneratorPlugin/Private/LuaScriptCodeGenerator.h b/Engine/Plugins/Script/ScriptGeneratorPlugin/Source/ScriptGeneratorPlugin/Private/LuaScriptCodeGenerator.h
--- a/Engine/Plugins/Script/ScriptGeneratorPlugin/Source/ScriptGeneratorPlugin/Private/LuaScriptCodeGenerator.h
+++ b/Engine/Plugins/Script/ScriptGeneratorPlugin/Source/ScriptGeneratorPlugin/Private/LuaScriptCodeGenerator.h
@@ -55,7 +55,7 @@ protected:
 
 public:
 
-	FLuaScriptCodeGenerator(const FString& RootLocalPath, const FString& RootBuildPath, const FString& OutputDirectory);
+	FLuaScriptCodeGenerator(const FString& RootLocalPath, const FString& RootBuildPath, const FString& OutputDirectory, const FString& InIncludeBase);
 
 	// FScriptCodeGeneratorBase interface
 	virtual void ExportClass(UClass* Class, const FString& SourceHeaderFilename, const FString& GeneratedHeaderFilename, bool bHasChanged) override;
diff --git a/Engine/Plugins/Script/ScriptGeneratorPlugin/Source/ScriptGeneratorPlugin/Private/ScriptCodeGeneratorBase.cpp b/Engine/Plugins/Script/ScriptGeneratorPlugin/Source/ScriptGeneratorPlugin/Private/ScriptCodeGeneratorBase.cpp
--- a/Engine/Plugins/Script/ScriptGeneratorPlugin/Source/ScriptGeneratorPlugin/Private/ScriptCodeGeneratorBase.cpp
+++ b/Engine/Plugins/Script/ScriptGeneratorPlugin/Source/ScriptGeneratorPlugin/Private/ScriptCodeGeneratorBase.cpp
@@ -2,11 +2,12 @@
 #include "ScriptGeneratorPluginPrivatePCH.h"
 #include "ScriptCodeGeneratorBase.h"
 
-FScriptCodeGeneratorBase::FScriptCodeGeneratorBase(const FString& InRootLocalPath, const FString& InRootBuildPath, const FString& InOutputDirectory)
+FScriptCodeGeneratorBase::FScriptCodeGeneratorBase(const FString& InRootLocalPath, const FString& InRootBuildPath, const FString& InOutputDirectory, const FString& InIncludeBase)
 {
 	GeneratedCodePath = InOutputDirectory;
 	RootLocalPath = InRootLocalPath;
 	RootBuildPath = InRootBuildPath;
+	IncludeBase = InIncludeBase;
 }
 
 bool FScriptCodeGeneratorBase::SaveHeaderIfChanged(const FString& HeaderPath, const FString& NewHeaderContents)
@@ -55,13 +56,7 @@ void FScriptCodeGeneratorBase::RenameTempFiles()
 FString FScriptCodeGeneratorBase::RebaseToBuildPath(const FString& FileName) const
 {
 	FString NewFilename(FileName);
-	if (RootLocalPath != RootBuildPath)
-	{
-		if (FPaths::MakePathRelativeTo(NewFilename, *RootLocalPath))
-		{
-			NewFilename = RootBuildPath / NewFilename;
-		}
-	}
+	FPaths::MakePathRelativeTo(NewFilename, *IncludeBase);
 	return NewFilename;
 }
 
diff --git a/Engine/Plugins/Script/ScriptGeneratorPlugin/Source/ScriptGeneratorPlugin/Private/ScriptCodeGeneratorBase.h b/Engine/Plugins/Script/ScriptGeneratorPlugin/Source/ScriptGeneratorPlugin/Private/ScriptCodeGeneratorBase.h
--- a/Engine/Plugins/Script/ScriptGeneratorPlugin/Source/ScriptGeneratorPlugin/Private/ScriptCodeGeneratorBase.h
+++ b/Engine/Plugins/Script/ScriptGeneratorPlugin/Source/ScriptGeneratorPlugin/Private/ScriptCodeGeneratorBase.h
@@ -18,11 +18,13 @@ protected:
 	FString RootLocalPath;
 	/** Build root path - may be different to RootBuildPath if we're building remotely. **/
 	FString RootBuildPath;
+	/** Base include directory */
+	FString IncludeBase;
 	/** Set of all exported class names */
 	TSet<FName> ExportedClasses;
 
 	/** Default ctor */
-	FScriptCodeGeneratorBase(const FString& InRootLocalPath, const FString& InRootBuildPath, const FString& OutputDirectory);
+	FScriptCodeGeneratorBase(const FString& InRootLocalPath, const FString& InRootBuildPath, const FString& OutputDirectory, const FString& InIncludeBase);
 
 	/** Saves generated script glue heade to a temporary file if its contents is different from the eexisting one. */
 	bool SaveHeaderIfChanged(const FString& HeaderPath, const FString& NewHeaderContents);
diff --git a/Engine/Plugins/Script/ScriptGeneratorPlugin/Source/ScriptGeneratorPlugin/Private/ScriptGeneratorPlugin.cpp b/Engine/Plugins/Script/ScriptGeneratorPlugin/Source/ScriptGeneratorPlugin/Private/ScriptGeneratorPlugin.cpp
--- a/Engine/Plugins/Script/ScriptGeneratorPlugin/Source/ScriptGeneratorPlugin/Private/ScriptGeneratorPlugin.cpp
+++ b/Engine/Plugins/Script/ScriptGeneratorPlugin/Source/ScriptGeneratorPlugin/Private/ScriptGeneratorPlugin.cpp
@@ -20,7 +20,7 @@ class FScriptGeneratorPlugin : public IScriptGeneratorPlugin
 	virtual FString GetGeneratedCodeModuleName() const override { return TEXT("ScriptPlugin"); }
 	virtual bool ShouldExportClassesForModule(const FString& ModuleName, EBuildModuleType::Type ModuleType) const;
 	virtual bool SupportsTarget(const FString& TargetName) const override { return true; }
-	virtual void Initialize(const FString& RootLocalPath, const FString& RootBuildPath, const FString& OutputDirectory) override;
+	virtual void Initialize(const FString& RootLocalPath, const FString& RootBuildPath, const FString& OutputDirectory, const FString& IncludeBase) override;
 	virtual void ExportClass(UClass* Class, const FString& SourceHeaderFilename, const FString& GeneratedHeaderFilename, bool bHasChanged) override;
 	virtual void FinishExport() override;
 };
@@ -37,14 +37,14 @@ void FScriptGeneratorPlugin::ShutdownModule()
 	CodeGenerator.Reset();
 }
 
-void FScriptGeneratorPlugin::Initialize(const FString& RootLocalPath, const FString& RootBuildPath, const FString& OutputDirectory)
+void FScriptGeneratorPlugin::Initialize(const FString& RootLocalPath, const FString& RootBuildPath, const FString& OutputDirectory, const FString& IncludeBase)
 {
 #if WITH_LUA
 	UE_LOG(LogScriptGenerator, Log, TEXT("Using Lua Script Generator."));
-	CodeGenerator = new FLuaScriptCodeGenerator(RootLocalPath, RootBuildPath, OutputDirectory);
+	CodeGenerator = new FLuaScriptCodeGenerator(RootLocalPath, RootBuildPath, OutputDirectory, IncludeBase);
 #else
 	UE_LOG(LogScriptGenerator, Log, TEXT("Using Generic Script Generator."));
-	CodeGenerator = new FGenericScriptCodeGenerator(RootLocalPath, RootBuildPath, OutputDirectory);
+	CodeGenerator = new FGenericScriptCodeGenerator(RootLocalPath, RootBuildPath, OutputDirectory, IncludeBase);
 #endif
 	UE_LOG(LogScriptGenerator, Log, TEXT("Output directory: %s"), *OutputDirectory);
 }
diff --git a/Engine/Source/Programs/UnrealHeaderTool/Private/CodeGenerator.cpp b/Engine/Source/Programs/UnrealHeaderTool/Private/CodeGenerator.cpp
--- a/Engine/Source/Programs/UnrealHeaderTool/Private/CodeGenerator.cpp
+++ b/Engine/Source/Programs/UnrealHeaderTool/Private/CodeGenerator.cpp
@@ -4387,23 +4387,25 @@ void GetScriptPlugins(TArray<IScriptGeneratorPluginInterface*>& ScriptPlugins)
 		if (bSupportedPlugin)
 		{
 			// Find the right output direcotry for this plugin base on its target (Engine-side) plugin name.
-			FString GeneratedCodeModule = ScriptGenerator->GetGeneratedCodeModuleName();
+			FString GeneratedCodeModuleName = ScriptGenerator->GetGeneratedCodeModuleName();
+			const FManifestModule* GeneratedCodeModule = NULL;
 			FString OutputDirectory;
+			FString IncludeBase;
 			for (const auto& Module : GManifest.Modules)
 			{
-				if (Module.Name == GeneratedCodeModule)
+				if (Module.Name == GeneratedCodeModuleName)
 				{
-					OutputDirectory = Module.GeneratedIncludeDirectory;
+					GeneratedCodeModule = &Module;
 				}
 			}
-			if (!OutputDirectory.IsEmpty())
+			if (GeneratedCodeModule)
 			{
-				ScriptGenerator->Initialize(GManifest.RootLocalPath, GManifest.RootBuildPath, OutputDirectory);
+				ScriptGenerator->Initialize(GManifest.RootLocalPath, GManifest.RootBuildPath, GeneratedCodeModule->GeneratedIncludeDirectory, GeneratedCodeModule->IncludeBase);
 			}
 			else
 			{
 				// Can't use this plugin
-				UE_LOG(LogCompile, Log, TEXT("Unable to determine output directory for %s. Cannot export script glue."), *GeneratedCodeModule);
+				UE_LOG(LogCompile, Log, TEXT("Unable to determine output directory for %s. Cannot export script glue."), *GeneratedCodeModuleName);
 				bSupportedPlugin = false;				
 			}
 		}
diff --git a/Engine/Source/Programs/UnrealHeaderTool/Public/IScriptGeneratorPluginInterface.h b/Engine/Source/Programs/UnrealHeaderTool/Public/IScriptGeneratorPluginInterface.h
--- a/Engine/Source/Programs/UnrealHeaderTool/Public/IScriptGeneratorPluginInterface.h
+++ b/Engine/Source/Programs/UnrealHeaderTool/Public/IScriptGeneratorPluginInterface.h
@@ -80,11 +80,10 @@ public:
 	/** Returns true if this plugin supports exporting scripts for the specified module */
 	virtual bool ShouldExportClassesForModule(const FString& ModuleName, EBuildModuleType::Type ModuleType) const = 0;
 	/** Initializes this plugin with build information */
-	virtual void Initialize(const FString& RootLocalPath, const FString& RootBuildPath, const FString& OutputDirectory) = 0;
+	virtual void Initialize(const FString& RootLocalPath, const FString& RootBuildPath, const FString& OutputDirectory, const FString& IncludeBase) = 0;
 	/** Exports a single class. May be called multiple times for the same class (as UHT processes the entire hierarchy inside modules. */
 	virtual void ExportClass(class UClass* Class, const FString& SourceHeaderFilename, const FString& GeneratedHeaderFilename, bool bHasChanged) = 0;
 	/** Called once all classes have been exported */
 	virtual void FinishExport() = 0;
-
 };
 
