ScriptPlugin: Don't require specific types for UObject properties in generated code (CIS Fix)

diff --git a/Engine/Plugins/Script/ScriptGeneratorPlugin/Source/ScriptGeneratorPlugin/Private/GenericScriptCodeGenerator.cpp b/Engine/Plugins/Script/ScriptGeneratorPlugin/Source/ScriptGeneratorPlugin/Private/GenericScriptCodeGenerator.cpp
--- a/Engine/Plugins/Script/ScriptGeneratorPlugin/Source/ScriptGeneratorPlugin/Private/GenericScriptCodeGenerator.cpp
+++ b/Engine/Plugins/Script/ScriptGeneratorPlugin/Source/ScriptGeneratorPlugin/Private/GenericScriptCodeGenerator.cpp
@@ -20,21 +20,21 @@ FString FGenericScriptCodeGenerator::GenerateWrapperFunctionDeclaration(const FS
 
 FString FGenericScriptCodeGenerator::GenerateFunctionParamDeclaration(const FString& ClassNameCPP, UClass* Class, UFunction* Function, UProperty* Param)
 {
-	FString Initializer;
+	FString ParamDecl;
 	if (Param->IsA(UObjectPropertyBase::StaticClass()) || Param->IsA(UClassProperty::StaticClass()))
 	{
-		Initializer = TEXT("NULL");
+		ParamDecl = FString::Printf(TEXT("UObject* %s = nullptr;"), *Param->GetName());
 	}
 	else
 	{
-		Initializer = FString::Printf(TEXT("%s()"), *GetPropertyTypeCPP(Param, CPPF_ArgumentOrReturnValue));
+		ParamDecl = FString::Printf(TEXT("%s %s = %s();"), *GetPropertyTypeCPP(Param, CPPF_ArgumentOrReturnValue), *Param->GetName(), *GetPropertyTypeCPP(Param, CPPF_ArgumentOrReturnValue));
 	}
-	return FString::Printf(TEXT("%s %s = %s;"), *GetPropertyTypeCPP(Param, CPPF_ArgumentOrReturnValue), *Param->GetName(), *Initializer);
+	return ParamDecl;
 }
 
 FString FGenericScriptCodeGenerator::GenerateObjectDeclarationFromContext(const FString& ClassNameCPP, UClass* Class)
 {
-	return FString::Printf(TEXT("UObject* Obj = (%s*)InScriptContext;"), *ClassNameCPP);
+	return FString::Printf(TEXT("UObject* Obj = (UObject*)InScriptContext;"));
 }
 
 FString FGenericScriptCodeGenerator::GenerateReturnValueHandler(const FString& ClassNameCPP, UClass* Class, UFunction* Function, UProperty* ReturnValue)
