[GitHub] 337 : Add module gen include dir to ShouldExportClassesForModule

#ttp 343350

diff --git a/Engine/Plugins/ScriptGeneratorPlugin/Source/ScriptGeneratorPlugin/Private/ScriptGeneratorPlugin.cpp b/Engine/Plugins/ScriptGeneratorPlugin/Source/ScriptGeneratorPlugin/Private/ScriptGeneratorPlugin.cpp
--- a/Engine/Plugins/ScriptGeneratorPlugin/Source/ScriptGeneratorPlugin/Private/ScriptGeneratorPlugin.cpp
+++ b/Engine/Plugins/ScriptGeneratorPlugin/Source/ScriptGeneratorPlugin/Private/ScriptGeneratorPlugin.cpp
@@ -18,7 +18,7 @@ class FScriptGeneratorPlugin : public IScriptGeneratorPlugin
 
 	/** IScriptGeneratorPlugin interface */
 	virtual FString GetGeneratedCodeModuleName() const override { return TEXT("ScriptPlugin"); }
-	virtual bool ShouldExportClassesForModule(const FString& ModuleName, EBuildModuleType::Type ModuleType) const;
+	virtual bool ShouldExportClassesForModule(const FString& ModuleName, EBuildModuleType::Type ModuleType, const FString& ModuleGeneratedIncludeDirectory) const;
 	virtual bool SupportsTarget(const FString& TargetName) const override { return true; }
 	virtual void Initialize(const FString& RootLocalPath, const FString& RootBuildPath, const FString& OutputDirectory, const FString& IncludeBase) override;
 	virtual void ExportClass(UClass* Class, const FString& SourceHeaderFilename, const FString& GeneratedHeaderFilename, bool bHasChanged) override;
@@ -49,7 +49,7 @@ void FScriptGeneratorPlugin::Initialize(const FString& RootLocalPath, const FStr
 	UE_LOG(LogScriptGenerator, Log, TEXT("Output directory: %s"), *OutputDirectory);
 }
 
-bool FScriptGeneratorPlugin::ShouldExportClassesForModule(const FString& ModuleName, EBuildModuleType::Type ModuleType) const
+bool FScriptGeneratorPlugin::ShouldExportClassesForModule(const FString& ModuleName, EBuildModuleType::Type ModuleType, const FString& ModuleGeneratedIncludeDirectory) const
 { 
 	bool bCanExport = (ModuleType == EBuildModuleType::Runtime || ModuleType == EBuildModuleType::Game);
 	if (bCanExport)
diff --git a/Engine/Source/Programs/UnrealHeaderTool/Private/HeaderParser.cpp b/Engine/Source/Programs/UnrealHeaderTool/Private/HeaderParser.cpp
--- a/Engine/Source/Programs/UnrealHeaderTool/Private/HeaderParser.cpp
+++ b/Engine/Source/Programs/UnrealHeaderTool/Private/HeaderParser.cpp
@@ -6771,7 +6771,7 @@ ECompilationResult::Type FHeaderParser::ParseAllHeadersInside(FClasses& ModuleCl
 		auto RootNode = &ModuleClasses.GetClassTree();
 		for (auto Plugin : ScriptPlugins)
 		{
-			if (Plugin->ShouldExportClassesForModule(Module.Name, Module.ModuleType))
+			if (Plugin->ShouldExportClassesForModule(Module.Name, Module.ModuleType, Module.GeneratedIncludeDirectory))
 			{
 				ExportClassToScriptPlugins(RootNode->GetClass(), Module, *Plugin);
 				ExportClassTreeToScriptPlugins(RootNode, Module, *Plugin);
diff --git a/Engine/Source/Programs/UnrealHeaderTool/Public/IScriptGeneratorPluginInterface.h b/Engine/Source/Programs/UnrealHeaderTool/Public/IScriptGeneratorPluginInterface.h
--- a/Engine/Source/Programs/UnrealHeaderTool/Public/IScriptGeneratorPluginInterface.h
+++ b/Engine/Source/Programs/UnrealHeaderTool/Public/IScriptGeneratorPluginInterface.h
@@ -78,7 +78,7 @@ public:
 	/** Returns true if this plugin supports exporting scripts for the specified target. This should handle game as well as editor target names */
 	virtual bool SupportsTarget(const FString& TargetName) const = 0;
 	/** Returns true if this plugin supports exporting scripts for the specified module */
-	virtual bool ShouldExportClassesForModule(const FString& ModuleName, EBuildModuleType::Type ModuleType) const = 0;
+	virtual bool ShouldExportClassesForModule(const FString& ModuleName, EBuildModuleType::Type ModuleType, const FString& ModuleGeneratedIncludeDirectory) const = 0;
 	/** Initializes this plugin with build information */
 	virtual void Initialize(const FString& RootLocalPath, const FString& RootBuildPath, const FString& OutputDirectory, const FString& IncludeBase) = 0;
 	/** Exports a single class. May be called multiple times for the same class (as UHT processes the entire hierarchy inside modules. */
