ScriptGeneratorPlugin will now be disabled by default. Projects can opt-in by enabling it in their uproject file.

diff --git a/Engine/Plugins/ScriptGeneratorPlugin/Source/ScriptGeneratorPlugin/Private/ScriptGeneratorPlugin.cpp b/Engine/Plugins/ScriptGeneratorPlugin/Source/ScriptGeneratorPlugin/Private/ScriptGeneratorPlugin.cpp
--- a/Engine/Plugins/ScriptGeneratorPlugin/Source/ScriptGeneratorPlugin/Private/ScriptGeneratorPlugin.cpp
+++ b/Engine/Plugins/ScriptGeneratorPlugin/Source/ScriptGeneratorPlugin/Private/ScriptGeneratorPlugin.cpp
@@ -4,6 +4,7 @@
 #include "ScriptCodeGeneratorBase.h"
 #include "GenericScriptCodeGenerator.h"
 #include "LuaScriptCodeGenerator.h"
+#include "IProjectManager.h"
 
 DEFINE_LOG_CATEGORY(LogScriptGenerator);
 
@@ -19,7 +20,7 @@ class FScriptGeneratorPlugin : public IScriptGeneratorPlugin
 	/** IScriptGeneratorPlugin interface */
 	virtual FString GetGeneratedCodeModuleName() const override { return TEXT("ScriptPlugin"); }
 	virtual bool ShouldExportClassesForModule(const FString& ModuleName, EBuildModuleType::Type ModuleType, const FString& ModuleGeneratedIncludeDirectory) const;
-	virtual bool SupportsTarget(const FString& TargetName) const override { return true; }
+	virtual bool SupportsTarget(const FString& TargetName) const override;
 	virtual void Initialize(const FString& RootLocalPath, const FString& RootBuildPath, const FString& OutputDirectory, const FString& IncludeBase) override;
 	virtual void ExportClass(UClass* Class, const FString& SourceHeaderFilename, const FString& GeneratedHeaderFilename, bool bHasChanged) override;
 	virtual void FinishExport() override;
@@ -77,3 +78,27 @@ void FScriptGeneratorPlugin::FinishExport()
 {
 	CodeGenerator->FinishExport();
 }
+
+bool FScriptGeneratorPlugin::SupportsTarget(const FString& TargetName) const
+{
+	// We only support the target if it explicitly enables the required script plugin in its uproject file
+	bool bSupportsTarget = false;
+	if (FPaths::IsProjectFilePathSet())
+	{
+		FProjectDescriptor ProjectDescriptor;
+		FText OutError;
+		if (ProjectDescriptor.Load(FPaths::GetProjectFilePath(), OutError))
+		{
+			for (auto& PluginDescriptor : ProjectDescriptor.Plugins)
+			{
+				// For your own script plugin you might want to change the hardcoded name here to something else
+				if (PluginDescriptor.bEnabled && PluginDescriptor.Name == TEXT("ScriptPlugin"))
+				{
+					bSupportsTarget = true;
+					break;
+				}
+			}
+		}
+	}
+	return bSupportsTarget;
+}
\ No newline at end of file
diff --git a/Engine/Plugins/ScriptGeneratorPlugin/Source/ScriptGeneratorPlugin/ScriptGeneratorPlugin.Build.cs b/Engine/Plugins/ScriptGeneratorPlugin/Source/ScriptGeneratorPlugin/ScriptGeneratorPlugin.Build.cs
--- a/Engine/Plugins/ScriptGeneratorPlugin/Source/ScriptGeneratorPlugin/ScriptGeneratorPlugin.Build.cs
+++ b/Engine/Plugins/ScriptGeneratorPlugin/Source/ScriptGeneratorPlugin/ScriptGeneratorPlugin.Build.cs
@@ -35,6 +35,7 @@ namespace UnrealBuildTool.Rules
 				new string[]
 				{
 					// ... add private dependencies that you statically link with here ...
+					"Projects",
 				}
 				);
 
diff --git a/Engine/Plugins/ScriptPlugin/ScriptPlugin.uplugin b/Engine/Plugins/ScriptPlugin/ScriptPlugin.uplugin
--- a/Engine/Plugins/ScriptPlugin/ScriptPlugin.uplugin
+++ b/Engine/Plugins/ScriptPlugin/ScriptPlugin.uplugin
@@ -9,7 +9,7 @@
 	"EngineVersion" : "4.2.0",
 	"Description" : "An example of a script plugin. This can be used as a starting point when creating your own plugin.",
 	"Category" : "Examples",
-	"EnabledByDefault" : true,
+	"EnabledByDefault" : false,
 	
 	"Modules" :
 	[
diff --git a/Engine/Source/Programs/UnrealHeaderTool/Private/CodeGenerator.cpp b/Engine/Source/Programs/UnrealHeaderTool/Private/CodeGenerator.cpp
--- a/Engine/Source/Programs/UnrealHeaderTool/Private/CodeGenerator.cpp
+++ b/Engine/Source/Programs/UnrealHeaderTool/Private/CodeGenerator.cpp
@@ -4825,6 +4825,7 @@ void GetScriptPlugins(TArray<IScriptGeneratorPluginInterface*>& ScriptPlugins)
 			auto GeneratorInterface = FModuleManager::LoadModulePtr<IScriptGeneratorPluginInterface>(*PluginStatus.Name);
 			if (GeneratorInterface)
 			{
+				UE_LOG(LogCompile, Log, TEXT("Detected script generator plugin (%d): %s"), ScriptPlugins.Num(), *PluginStatus.Name);
 				ScriptPlugins.Add(GeneratorInterface);
 			}
 		}
@@ -4862,6 +4863,7 @@ void GetScriptPlugins(TArray<IScriptGeneratorPluginInterface*>& ScriptPlugins)
 		}
 		if (!bSupportedPlugin)
 		{
+			UE_LOG(LogCompile, Log, TEXT("Script generator plugin %d not supported for target: %s"), PluginIndex, *GManifest.TargetName);
 			ScriptPlugins.RemoveAt(PluginIndex);
 		}
 	}
